Extract DDC HttpClient code into DevHttp module.

diff --git a/Engine/Source/Developer/DerivedDataCache/DerivedDataCache.Build.cs b/Engine/Source/Developer/DerivedDataCache/DerivedDataCache.Build.cs
--- a/Engine/Source/Developer/DerivedDataCache/DerivedDataCache.Build.cs
+++ b/Engine/Source/Developer/DerivedDataCache/DerivedDataCache.Build.cs
@@ -8,6 +8,7 @@ public class DerivedDataCache : ModuleRules
 	public DerivedDataCache(ReadOnlyTargetRules Target) : base(Target)
 	{
 		PrivateDependencyModuleNames.Add("Core");
+		PrivateDependencyModuleNames.Add("DevHttp");
 		PrivateDependencyModuleNames.Add("Zen");
 		
 		// Dependencies for "S3" and "HTTP" backends
diff --git a/Engine/Source/Developer/DerivedDataCache/Private/Http/CurlHttpClient.cpp b/Engine/Source/Developer/DerivedDataCache/Private/Http/CurlHttpClient.cpp
deleted file mode 100644
diff --git a/Engine/Source/Developer/DerivedDataCache/Private/Http/HttpClient.cpp b/Engine/Source/Developer/DerivedDataCache/Private/Http/HttpClient.cpp
deleted file mode 100644
diff --git a/Engine/Source/Developer/DerivedDataCache/Private/Http/HttpClient.h b/Engine/Source/Developer/DerivedDataCache/Private/Http/HttpClient.h
deleted file mode 100644
diff --git a/Engine/Source/Developer/DerivedDataCache/Private/Http/HttpHostBuilder.cpp b/Engine/Source/Developer/DerivedDataCache/Private/Http/HttpHostBuilder.cpp
deleted file mode 100644
diff --git a/Engine/Source/Developer/DerivedDataCache/Private/Http/HttpHostBuilder.h b/Engine/Source/Developer/DerivedDataCache/Private/Http/HttpHostBuilder.h
deleted file mode 100644
diff --git a/Engine/Source/Developer/DevHttp/DevHttp.Build.cs b/Engine/Source/Developer/DevHttp/DevHttp.Build.cs
new file mode 100644
--- /dev/null
+++ b/Engine/Source/Developer/DevHttp/DevHttp.Build.cs
@@ -0,0 +1,15 @@
+// Copyright Epic Games, Inc. All Rights Reserved.
+
+using UnrealBuildTool;
+using System.IO;
+
+public class DevHttp : ModuleRules
+{
+	public DevHttp(ReadOnlyTargetRules Target) : base(Target)
+	{
+		PrivateDependencyModuleNames.AddRange(new string[] { "Core", "SSL" });
+		
+		AddEngineThirdPartyPrivateStaticDependencies(Target, "libcurl");
+		AddEngineThirdPartyPrivateStaticDependencies(Target, "OpenSSL");
+	}
+}
diff --git a/Engine/Source/Developer/DevHttp/Internal/Http/HttpClient.h b/Engine/Source/Developer/DevHttp/Internal/Http/HttpClient.h
new file mode 100644
--- /dev/null
+++ b/Engine/Source/Developer/DevHttp/Internal/Http/HttpClient.h
@@ -0,0 +1,488 @@
+// Copyright Epic Games, Inc. All Rights Reserved.
+
+#pragma once
+
+#include "CoreTypes.h"
+#include "Containers/ContainersFwd.h"
+#include "Containers/StringFwd.h"
+#include "Containers/StringView.h"
+#include "Memory/MemoryFwd.h"
+#include "Templates/Function.h"
+#include "Templates/PimplPtr.h"
+#include "Templates/RefCounting.h"
+#include "Templates/UniquePtr.h"
+
+namespace UE { class IHttpClient; }
+namespace UE { class IHttpConnectionPool; }
+namespace UE { class IHttpReceiver; }
+namespace UE { class IHttpRequest; }
+namespace UE { class IHttpResponse; }
+namespace UE { struct FHttpClientParams; }
+namespace UE { struct FHttpConnectionPoolParams; }
+namespace UE { struct FHttpRequestParams; }
+namespace UE { struct FHttpResponseStats; }
+namespace UE::Http::Private { struct FHttpRequestQueueData; }
+
+namespace UE::Http::Private
+{
+
+template <typename T>
+struct THttpDestroyer
+{
+	void operator()(void* Object) const
+	{
+		if (Object)
+		{
+			((T*)Object)->Destroy();
+		}
+	}
+};
+
+} // UE::Http::Private
+
+namespace UE
+{
+
+template <typename Type, typename BaseType = Type>
+using THttpUniquePtr = TUniquePtr<Type, Http::Private::THttpDestroyer<BaseType>>;
+
+enum class EHttpMethod : uint8
+{
+	Get,
+	Put,
+	Post,
+	Head,
+	Delete,
+};
+
+DEVHTTP_API FAnsiStringView LexToString(EHttpMethod Method);
+DEVHTTP_API bool TryLexFromString(EHttpMethod& OutMethod, FAnsiStringView View);
+
+enum class EHttpMediaType : uint8
+{
+	Any,
+	Binary,
+	Text,
+	Json,
+	Yaml,
+	CbObject,
+	CbPackage,
+	CbPackageOffer,
+	CompressedBinary,
+	FormUrlEncoded,
+};
+
+DEVHTTP_API FAnsiStringView LexToString(EHttpMediaType MediaType);
+DEVHTTP_API bool TryLexFromString(EHttpMediaType& OutMediaType, FAnsiStringView View);
+
+enum class EHttpVersion : uint8
+{
+	/** Determine version automatically. */
+	None = 0,
+	/** Use HTTP/1.0 */
+	V1_0,
+	/** Use HTTP/1.1 */
+	V1_1,
+	/** Try HTTP/2 with fallback to HTTP/1.1 */
+	V2,
+	/** Use HTTP/2 */
+	V2Only,
+};
+
+enum class EHttpTlsLevel : uint8
+{
+	/** Do not use TLS. */
+	None = 0,
+	/** Try to use TLS and fall back to none. */
+	Try,
+	/** Require TLS. */
+	All,
+};
+
+enum class EHttpErrorCode : uint8
+{
+	/** Success. */
+	None = 0,
+	/** Unknown. Check IHttpResponse::GetError(). */
+	Unknown,
+	/** Request was canceled. */
+	Canceled,
+	/** Failed to resolve the host. */
+	ResolveHost,
+	/** Failed to establish a connection. */
+	Connect,
+	/** Failed to establish a TLS connection. */
+	TlsConnect,
+	/** Failed to verify the certificate of the peer. */
+	TlsPeerVerification,
+	/** Failed because the timeout period was exceeded. */
+	TimedOut,
+};
+
+class IHttpManager
+{
+public:
+	static DEVHTTP_API IHttpManager& Get();
+
+	[[nodiscard]] virtual THttpUniquePtr<IHttpConnectionPool> CreateConnectionPool(const FHttpConnectionPoolParams& Params) = 0;
+};
+
+class IHttpConnectionPool
+{
+public:
+	[[nodiscard]] virtual THttpUniquePtr<IHttpClient> CreateClient(const FHttpClientParams& Params) = 0;
+
+protected:
+	friend Http::Private::THttpDestroyer<IHttpConnectionPool>;
+	virtual ~IHttpConnectionPool() = default;
+	virtual void Destroy() = 0;
+};
+
+struct FHttpConnectionPoolParams final
+{
+	/** Maximum number of concurrent connections created by the pool. Use 0 for the default limit. */
+	uint32 MaxConnections = 0;
+
+	/** Minimum number of concurrent connections maintained by the pool for reuse. Use 0 for the default limit. */
+	uint32 MinConnections = 0;
+
+	/** Maximum number of concurrent requests in flight for a connection when multiplexing. Use 0 for the default limit. */
+	uint32 MaxRequestsPerConnection = 0;
+
+	/** Allow requests to send async using the pool. Requests will block on execution when this is disabled. */
+	bool bAllowAsync = true;
+};
+
+class IHttpClient
+{
+public:
+	/**
+	 * Try to create a request.
+	 *
+	 * The request is counted against the maximum request count as soon as it is created.
+	 *
+	 * Returns null when the maximum request count has been reached. Retry after destroying a request.
+	 */
+	[[nodiscard]] virtual THttpUniquePtr<IHttpRequest> TryCreateRequest(const FHttpRequestParams& Params) = 0;
+
+protected:
+	friend Http::Private::THttpDestroyer<IHttpClient>;
+	virtual ~IHttpClient() = default;
+	virtual void Destroy() = 0;
+};
+
+struct FHttpClientParams final
+{
+	/** Invoked whenever a request from this client is destroyed. This may be used to retry creating a request. */
+	TFunction<void()> OnDestroyRequest;
+
+	/** Maximum number of concurrent requests created by the client. Use 0 for the default limit. */
+	uint32 MaxRequests = 0;
+
+	/** Minimum number of concurrent requests maintained by the client for reuse. Use 0 for the default limit. */
+	uint32 MinRequests = 0;
+
+	/** Domain name cache timeout in seconds. Use 0 to disable the cache. */
+	uint32 DnsCacheTimeout = 60;
+
+	/** Connection timeout in milliseconds. Use 0 for the default timeout. */
+	uint32 ConnectTimeout = 0;
+
+	/** Average transfer speed in bytes/s below which a request will abort. Use 0 to disable. */
+	uint32 LowSpeedLimit = 0;
+
+	/** Time in seconds after which a request will abort if it stays below LowSpeedLimit. Use 0 to disable. */
+	uint32 LowSpeedTime = 0;
+
+	/** HTTP version to use for requests created by the client. */
+	EHttpVersion Version = EHttpVersion::None;
+
+	/** Level of TLS to use for requests created by the client. */
+	EHttpTlsLevel TlsLevel = EHttpTlsLevel::None;
+
+	/** Follow redirects in responses automatically. */
+	bool bFollowRedirects : 1;
+	/** Follow redirects of the corresponding status code without rewriting POST to GET. */
+	bool bFollow301Post : 1;
+	bool bFollow302Post : 1;
+	bool bFollow303Post : 1;
+
+	/** Avoid use of configured proxy on the client. */
+	bool bBypassProxy : 1;
+
+	/** Verbose logging for requests created by the client. */
+	bool bVerbose : 1;
+
+	inline FHttpClientParams()
+		: bFollowRedirects(false)
+		, bFollow301Post(false)
+		, bFollow302Post(false)
+		, bFollow303Post(false)
+		, bBypassProxy(false)
+		, bVerbose(false)
+	{
+	}
+};
+
+class IHttpRequest
+{
+public:
+	/**
+	 * Reset the request to an empty state equivalent to when it was created.
+	 *
+	 * The request must be idle to be reset, which means any associated response must have completed.
+	 */
+	virtual void Reset() = 0;
+
+	/** Set the URI used for the request, including the scheme and authority. */
+	virtual void SetUri(FAnsiStringView Uri) = 0;
+
+	/** Set a Unix socket path that should be used for sending requests */
+	virtual void SetUnixSocketPath(FAnsiStringView SocketPath) = 0;
+
+	/** Set the method used for the request. Defaults to GET. */
+	virtual void SetMethod(EHttpMethod Method) = 0;
+
+	/**
+	 * Set the body for PUT and POST requests.
+	 *
+	 * Body must be owned or otherwise remain valid until the request is complete.
+	 */
+	virtual void SetBody(const FCompositeBuffer& Body) = 0;
+
+	/** Set the Content-Type header with a known media type. Use AddHeader() for other types. */
+	DEVHTTP_API void SetContentType(EHttpMediaType Type, FAnsiStringView Param = {});
+
+	/** Add an Accept header with a known media type. Use AddHeader() for other types. */
+	DEVHTTP_API void AddAcceptType(EHttpMediaType Type, float Weight = 1.0);
+
+	/** Add the header. An empty value removes the header. */
+	virtual void AddHeader(FAnsiStringView Name, FAnsiStringView Value) = 0;
+
+	/**
+	 * Send the request on the client that it was created from, and wait for it to complete before returning.
+	 *
+	 * Request must be idle to send, which means any previous response must be complete.
+	 *
+	 * Receiver must remain valid until OnComplete is called or until it returns another receiver
+	 * to be used for subsequent events.
+	 */
+	virtual void Send(IHttpReceiver* Receiver, THttpUniquePtr<IHttpResponse>& OutResponse) = 0;
+
+	/**
+	 * Asynchronously send the request on the client that it was created from.
+	 *
+	 * Request must be idle to send, which means any previous response must be complete.
+	 *
+	 * Receiver must remain valid until OnComplete is called or until it returns another receiver
+	 * to be used for subsequent events.
+	 *
+	 * Response must be kept alive until OnComplete is called. Cancel to complete immediately.
+	 */
+	virtual void SendAsync(IHttpReceiver* Receiver, THttpUniquePtr<IHttpResponse>& OutResponse) = 0;
+
+protected:
+	friend Http::Private::THttpDestroyer<IHttpRequest>;
+	virtual ~IHttpRequest() = default;
+	virtual void Destroy() = 0;
+};
+
+struct FHttpRequestParams final
+{
+	/** Whether to force creation of the request even if it would exceed the configured maximum. */
+	bool bIgnoreMaxRequests = false;
+};
+
+class IHttpResponseMonitor
+{
+public:
+	/** Cancels the request associated with this response. */
+	virtual void Cancel() = 0;
+
+	/** Waits for the request associated with this response to be complete. */
+	virtual void Wait() const = 0;
+
+	/** Returns true if the request associated with this response is complete, and false otherwise. */
+	[[nodiscard]] virtual bool Poll() const = 0;
+
+	virtual void AddRef() const = 0;
+	virtual void Release() const = 0;
+};
+
+class IHttpResponse
+{
+public:
+	/** Returns a monitor for the response that can poll, wait, or cancel it. */
+	[[nodiscard]] virtual TRefCountPtr<IHttpResponseMonitor> GetMonitor() = 0;
+
+	/** Cancels the request associated with this response. */
+	virtual void Cancel() = 0;
+
+	/** Waits for the request associated with this response to be complete. */
+	virtual void Wait() const = 0;
+
+	/** Returns true if the request associated with this response is complete, and false otherwise. */
+	[[nodiscard]] virtual bool Poll() const = 0;
+
+	/** Get the URI used for the request associated with this response. Always available. */
+	[[nodiscard]] virtual FAnsiStringView GetUri() const = 0;
+
+	/** Get the method used for the request associated with this response. Always available. */
+	[[nodiscard]] virtual EHttpMethod GetMethod() const = 0;
+
+	/**
+	 * Returns the status code of the response. Available once the response status has been received.
+	 *
+	 * The return value is negative until the response status has been received.
+	 * The return value is 0 if an error occurred. See GetError().
+	 */
+	[[nodiscard]] virtual int32 GetStatusCode() const = 0;
+
+	/** Returns the error code from the HTTP stack. Available when the response is complete. */
+	[[nodiscard]] virtual EHttpErrorCode GetErrorCode() const = 0;
+
+	/** Returns the optional error message from the HTTP stack. Available when the status code is 0. */
+	[[nodiscard]] virtual FAnsiStringView GetError() const = 0;
+
+	/** Returns information about the request such as size, rate, and duration. Available when complete. */
+	[[nodiscard]] virtual const FHttpResponseStats& GetStats() const = 0;
+
+	/** Returns every header in the response. Available once OnHeaders() has been invoked. */
+	[[nodiscard]] virtual TConstArrayView<FAnsiStringView> GetAllHeaders() const = 0;
+
+	/** Returns the value of the first header matching the name, or empty if there was no matching header. */
+	[[nodiscard]] DEVHTTP_API FAnsiStringView GetHeader(FAnsiStringView Name) const;
+
+	/**
+	 * Fills the array of header values and returns the total number of headers matching the name.
+	 *
+	 * The return value can be greater than the array size if there are more matching headers than the array can hold.
+	 */
+	[[nodiscard]] DEVHTTP_API int32 GetHeaders(FAnsiStringView Name, TArrayView<FAnsiStringView> OutValues) const;
+
+	/** Returns the content type of the response, or Any if there was no matching header. */
+	[[nodiscard]] DEVHTTP_API EHttpMediaType GetContentType() const;
+
+protected:
+	friend Http::Private::THttpDestroyer<IHttpResponse>;
+	virtual ~IHttpResponse() = default;
+
+	/** Cancel() if not yet invoking OnComplete(), and then Release(). */
+	virtual void Destroy() = 0;
+};
+
+/** Appends the URI, method, status code, and optional error from the response to the builder. */
+DEVHTTP_API FStringBuilderBase& operator<<(FStringBuilderBase& Builder, const IHttpResponse& Response);
+
+struct FHttpResponseStats final
+{
+	/** Total size of data sent to the server, in bytes. */
+	uint64 SendSize = 0;
+	/** Total size of data received from the server, in bytes. */
+	uint64 RecvSize = 0;
+	/** Average rate at which data was sent to the server, in bytes/s. */
+	uint64 SendRate = 0;
+	/** Average rate at which data was received from the server, in bytes/s. */
+	uint64 RecvRate = 0;
+
+	/** Time from the start of the request until the name was resolved, in seconds. */
+	double NameResolveTime = 0.0;
+	/** Time from the start of the request until the connection was established, in seconds. */
+	double ConnectTime = 0.0;
+	/** Time from the start of the request until the TLS connection was complete, in seconds. */
+	double TlsConnectTime = 0.0;
+	/** Time from the start of the request until the request was sent but before the first byte was received, in seconds. */
+	double PreTransferTime = 0.0;
+	/** Time from the start of the request until the first byte was received, in seconds. */
+	double StartTransferTime = 0.0;
+	/** Time from the start of the request until response was complete, in seconds. */
+	double TotalTime = 0.0;
+
+	/** Time from when the request was sent until the first byte was received, in seconds. If no byte has been received, latency is 0. */
+	double GetLatency() const
+	{
+		return FMath::Max(0.0f, StartTransferTime - PreTransferTime);
+	}
+};
+
+/**
+ * Interface to receive the headers and body of the response.
+ *
+ * Functions on this interface must be implemented to do the minimal work possible, because their
+ * execution can block other requests from being processed.
+ */
+class IHttpReceiver
+{
+public:
+	virtual ~IHttpReceiver() = default;
+
+	/**
+	 * Invoked when the response has been created, before it begins executing.
+	 *
+	 * Returns the receiver to use for the next part of the response and can return itself.
+	 */
+	virtual IHttpReceiver* OnCreate(IHttpResponse& Response) { return this; }
+
+	/**
+	 * Invoked when every header has been received for the response.
+	 *
+	 * Returns the receiver to use for the next part of the response and can return itself.
+	 */
+	virtual IHttpReceiver* OnHeaders(IHttpResponse& Response) { return this; }
+
+	/**
+	 * Invoked when a part of the body has been received.
+	 *
+	 * The Data parameter can be reassigned to point to the tail of the view that is not consumed
+	 * by this receiver, and will be passed to the receiver that is returned.
+	 *
+	 * Returns the receiver to use for the next part of the response and can return itself.
+	 */
+	virtual IHttpReceiver* OnBody(IHttpResponse& Response, FMemoryView& Data) { return this; }
+
+	/**
+	 * Invoked when the response is complete. Always invoked exactly once for each response.
+	 *
+	 * Check the status code to detect errors. Can retry using the same request.
+	 *
+	 * It is valid for implementations of this function to delete the response.
+	 *
+	 * Returns the next receiver to handle OnComplete, or null if there are no others.
+	 */
+	virtual IHttpReceiver* OnComplete(IHttpResponse& Response) { return this; }
+};
+
+///////////////////////////////////////////////////////////////////////////////////////////////////////////////////////
+
+class FHttpByteArrayReceiver final : public IHttpReceiver
+{
+public:
+	explicit DEVHTTP_API FHttpByteArrayReceiver(TArray64<uint8>& OutArray, IHttpReceiver* Next = nullptr);
+
+	DEVHTTP_API IHttpReceiver* OnBody(IHttpResponse& Response, FMemoryView& Data) final;
+	DEVHTTP_API IHttpReceiver* OnComplete(IHttpResponse& Response) final { return Next; }
+
+private:
+	TArray64<uint8>& Array;
+	IHttpReceiver* Next;
+};
+
+///////////////////////////////////////////////////////////////////////////////////////////////////////////////////////
+
+/** Manages a client by letting callers queue to receive a request when one is not immediately available. */
+class FHttpRequestQueue final
+{
+public:
+	/** Create a null request queue. Asserts if CreateRequest is called. */
+	DEVHTTP_API FHttpRequestQueue() = default;
+
+	DEVHTTP_API FHttpRequestQueue(IHttpConnectionPool& ConnectionPool, const FHttpClientParams& ClientParams);
+
+	/** Blocks until a request can be created on the underlying client. */
+	[[nodiscard]] DEVHTTP_API THttpUniquePtr<IHttpRequest> CreateRequest(const FHttpRequestParams& Params);
+
+private:
+	TPimplPtr<Http::Private::FHttpRequestQueueData> Data;
+};
+
+} // UE
diff --git a/Engine/Source/Developer/DevHttp/Internal/Http/HttpHostBuilder.h b/Engine/Source/Developer/DevHttp/Internal/Http/HttpHostBuilder.h
new file mode 100644
--- /dev/null
+++ b/Engine/Source/Developer/DevHttp/Internal/Http/HttpHostBuilder.h
@@ -0,0 +1,33 @@
+// Copyright Epic Games, Inc. All Rights Reserved.
+
+#pragma once
+
+#include "CoreTypes.h"
+#include "Containers/ContainersFwd.h"
+#include "Containers/StringFwd.h"
+#include "Containers/StringView.h"
+
+namespace UE
+{
+
+class FHttpHostBuilder
+{
+public:
+	// accept a semicolon seperated list of hosts to pick from
+	DEVHTTP_API void AddFromString(FAnsiStringView HostList);
+	DEVHTTP_API void AddFromString(FStringView HostList);
+
+	// a url that is expected to host a /status/peers endpoint with endpoints we should attempt to use
+	DEVHTTP_API void AddFromEndpoint(FAnsiStringView HostUrl, FAnsiStringView AccessToken);
+
+	// produces a string of all the host candidates used for display purposes
+	DEVHTTP_API FString GetHostCandidatesString() const;
+
+	DEVHTTP_API bool ResolveHost(double WarningTimeoutSeconds, double TimeoutSeconds, FAnsiStringBuilderBase& OutHost, double& OutLatency);
+private:
+	static bool BenchmarkHostList(TConstArrayView<FString> HostCandidates, double WarningTimeoutSeconds, double TimeoutSeconds, FAnsiStringBuilderBase& OutHost, double& OutLatency);
+
+	TArray<FString> HostCandidates;
+};
+
+} // UE
diff --git a/Engine/Source/Developer/DevHttp/Private/CurlHttpClient.cpp b/Engine/Source/Developer/DevHttp/Private/CurlHttpClient.cpp
new file mode 100644
--- /dev/null
+++ b/Engine/Source/Developer/DevHttp/Private/CurlHttpClient.cpp
@@ -0,0 +1,1335 @@
+// Copyright Epic Games, Inc. All Rights Reserved.
+
+#include "Http/HttpClient.h"
+
+#include "Async/InheritedContext.h"
+#include "Async/ManualResetEvent.h"
+#include "Containers/AnsiString.h"
+#include "Containers/ConsumeAllMpmcQueue.h"
+#include "Containers/LockFreeList.h"
+#include "Containers/StringView.h"
+#include "HAL/PlatformProcess.h"
+#include "HAL/Thread.h"
+#include "Memory/CompositeBuffer.h"
+#include "Memory/MemoryView.h"
+#include "Misc/App.h"
+#include "Misc/EngineVersion.h"
+#include "Misc/StringBuilder.h"
+#include "Serialization/CompactBinary.h"
+#include "String/Find.h"
+#include "Templates/RefCounting.h"
+#include "Templates/UnrealTemplate.h"
+
+#if WITH_SSL
+#include "Ssl.h"
+#include <openssl/ssl.h>
+#endif
+
+#ifndef CURL_NO_OLDIES
+#define CURL_NO_OLDIES
+#endif
+
+#if PLATFORM_MICROSOFT
+#include "Microsoft/AllowMicrosoftPlatformTypes.h"
+#endif
+#ifdef PLATFORM_CURL_INCLUDE
+#include PLATFORM_CURL_INCLUDE
+#else
+#include "curl/curl.h"
+#endif
+#if PLATFORM_MICROSOFT
+#include "Microsoft/HideMicrosoftPlatformTypes.h"
+#endif
+
+#include <atomic>
+
+namespace UE { class FCurlHttpClient; }
+namespace UE { class FCurlHttpRequest; }
+namespace UE { class FCurlHttpResponse; }
+
+namespace UE::CurlHttp::Private
+{
+
+template <typename T>
+static bool TryIncrement(std::atomic<T>& Value, const T Max)
+{
+	for (T Existing = Value.load(std::memory_order_relaxed);;)
+	{
+		if (Existing >= Max)
+		{
+			return false;
+		}
+		if (Value.compare_exchange_weak(Existing, Existing + 1, std::memory_order_relaxed))
+		{
+			return true;
+		}
+	}
+}
+
+template <typename T>
+static T AtomicFetchOr(std::atomic<T>& Atomic, const T Value, std::memory_order MemoryOrder = std::memory_order_seq_cst)
+{
+	for (T Existing = Atomic.load(MemoryOrder);;)
+	{
+		if (Atomic.compare_exchange_weak(Existing, Existing | Value, MemoryOrder))
+		{
+			return Existing;
+		}
+	}
+}
+
+template <typename T, typename PredicateType>
+static bool AtomicFetchOrIf(
+	std::atomic<T>& Atomic,
+	const T Value,
+	PredicateType Predicate,
+	std::memory_order MemoryOrder = std::memory_order_seq_cst)
+{
+	for (T Existing = Atomic.load(MemoryOrder);;)
+	{
+		if (!Predicate(Existing))
+		{
+			return false;
+		}
+		if (Atomic.compare_exchange_weak(Existing, Existing | Value, MemoryOrder))
+		{
+			return true;
+		}
+	}
+}
+
+} // UE::CurlHttp::Private
+
+///////////////////////////////////////////////////////////////////////////////////////////////////////////////////////
+
+namespace UE
+{
+
+DEFINE_LOG_CATEGORY_STATIC(LogHttp, Display, All);
+
+class FCurlHttpHeaders
+{
+public:
+	FCurlHttpHeaders() = default;
+	FCurlHttpHeaders(const FCurlHttpHeaders&) = delete;
+	FCurlHttpHeaders& operator=(const FCurlHttpHeaders&) = delete;
+
+	inline ~FCurlHttpHeaders()
+	{
+		curl_slist_free_all(List);
+	}
+
+	inline void Reset()
+	{
+		curl_slist_free_all(List);
+		List = nullptr;
+	}
+
+	inline void AddHeader(const ANSICHAR* Header)
+	{
+		List = curl_slist_append(List, Header);
+	}
+
+	inline curl_slist* GetList() const { return List; }
+
+private:
+	curl_slist* List = nullptr;
+};
+
+class FCurlHttpManager final : public IHttpManager
+{
+public:
+	FCurlHttpManager();
+
+	THttpUniquePtr<IHttpConnectionPool> CreateConnectionPool(const FHttpConnectionPoolParams& Params) final;
+
+	void SetDefaultOptions(CURL* Curl, FCurlHttpHeaders& Headers);
+
+private:
+	const FCbObjectId SessionId = FApp::GetSessionObjectId();
+	std::atomic<uint32> RequestId = 1;
+	FUtf8StringBuilderBase UserAgent;
+};
+
+class FCurlHttpConnectionPool final : public IHttpConnectionPool
+{
+public:
+	FCurlHttpConnectionPool(FCurlHttpManager& Manager, const FHttpConnectionPoolParams& Params);
+
+	THttpUniquePtr<IHttpClient> CreateClient(const FHttpClientParams& Params) final;
+	void DeleteClient(FCurlHttpClient* Client);
+
+	void SetDefaultOptions(CURL* Curl, FCurlHttpHeaders& Headers) const;
+
+	bool BeginAsyncRequest(FCurlHttpResponse* Response);
+	void CancelAsyncRequest(FCurlHttpResponse* Response);
+
+private:
+	~FCurlHttpConnectionPool() final;
+	void Destroy() final { delete this; }
+
+	void ThreadLoop();
+	void CompleteRequest(CURL* Curl, CURLcode* OptionalReturnCode);
+
+	static void CurlLock(CURL* Curl, curl_lock_data Data, curl_lock_access Access, void* Param);
+	static void CurlUnlock(CURL* Curl, curl_lock_data Data, void* Param);
+
+	static void AssertShareCodeOk(CURLSHcode Code);
+	static void AssertMultiCodeOk(CURLMcode Code);
+
+	FCurlHttpManager& Manager;
+	CURLSH* CurlShare;
+	CURLM* CurlMulti;
+
+	FRWLock Locks[CURL_LOCK_DATA_LAST];
+	bool WriteLocked[CURL_LOCK_DATA_LAST]{};
+
+	std::atomic<uint32> ClientCount = 0;
+
+	enum class EThreadCommandType
+	{
+		Begin,
+		Cancel,
+	};
+
+	struct FThreadCommand
+	{
+		TRefCountPtr<FCurlHttpResponse> Response;
+		EThreadCommandType Type;
+	};
+
+	TConsumeAllMpmcQueue<FThreadCommand> ThreadCommands;
+	FThread Thread;
+	std::atomic<bool> bThreadStarting;
+	std::atomic<bool> bThreadStopping;
+};
+
+class FCurlHttpClient final : public IHttpClient
+{
+public:
+	FCurlHttpClient(FCurlHttpConnectionPool& ConnectionPool, const FHttpClientParams& Params);
+
+	THttpUniquePtr<IHttpRequest> TryCreateRequest(const FHttpRequestParams& Params) final;
+	void DeleteRequest(FCurlHttpRequest* Request);
+
+	void SetDefaultOptions(CURL* Curl, FCurlHttpHeaders& Headers) const;
+
+	bool BeginAsyncRequest(FCurlHttpResponse* Response);
+	void CancelAsyncRequest(FCurlHttpResponse* Response);
+
+private:
+	~FCurlHttpClient() final;
+	void Destroy() final { delete this; }
+
+	static long ConvertVersion(EHttpVersion Version);
+	static long ConvertTlsLevel(EHttpTlsLevel Level);
+
+	FCurlHttpConnectionPool& ConnectionPool;
+	TLockFreePointerListLIFO<FCurlHttpRequest> RequestPool;
+	std::atomic<uint32> RequestPoolCount = 0;
+	std::atomic<uint32> RequestCount = 0;
+	FHttpClientParams Params;
+};
+
+class FCurlHttpRequest final : public IHttpRequest
+{
+public:
+	explicit FCurlHttpRequest(FCurlHttpClient& Client);
+	~FCurlHttpRequest() final;
+
+	void Reset() final;
+	void SetUri(FAnsiStringView Uri) final;
+	void SetUnixSocketPath(FAnsiStringView Uri) final;
+	void SetMethod(EHttpMethod Method) final;
+	void SetBody(const FCompositeBuffer& Body) final;
+	void AddHeader(FAnsiStringView Name, FAnsiStringView Value) final;
+
+	void Send(IHttpReceiver* Receiver, THttpUniquePtr<IHttpResponse>& OutResponse) final;
+	void SendAsync(IHttpReceiver* Receiver, THttpUniquePtr<IHttpResponse>& OutResponse) final;
+
+	void OnComplete(CURLcode Code);
+
+private:
+	void Destroy() final { Client.DeleteRequest(this); }
+
+	void SetDefaultOptions();
+	void CheckIdle(const TCHAR* FunctionName);
+
+	FCurlHttpResponse* CreateResponse(IHttpReceiver* Receiver, THttpUniquePtr<IHttpResponse>& OutResponse);
+
+	static size_t CurlRead(void* Ptr, size_t SizeInBlocks, size_t BlockSizeInBytes, void* Param);
+	static size_t CurlSeek(void* Param, curl_off_t Offset, int Origin);
+#if WITH_SSL
+	static CURLcode CurlSslContext(CURL* Curl, void* Context, void* Param);
+	static int SslCertVerify(int PreverifyOk, X509_STORE_CTX* Context);
+#endif
+
+	EHttpMethod Method;
+	FAnsiStringBuilderBase Uri;
+	FAnsiString UnixSocketPath;
+	FCurlHttpClient& Client;
+	FCurlHttpHeaders Headers;
+	FCompositeBuffer Body;
+	uint64 BodyOffset = 0;
+	CURL* Curl = nullptr;
+	std::atomic<FCurlHttpResponse*> Response = nullptr;
+};
+
+enum class ECurlHttpResponseState : uint8
+{
+	None = 0,
+	Complete = 1 << 0,
+	Canceled = 1 << 1,
+};
+
+ENUM_CLASS_FLAGS(ECurlHttpResponseState);
+
+class FCurlHttpResponse final : public IHttpResponse, public IHttpResponseMonitor, private FInheritedContextBase
+{
+public:
+	FCurlHttpResponse(CURL* Curl, EHttpMethod Method, FAnsiStringView Uri, IHttpReceiver* Receiver);
+
+	bool Create();
+
+	CURL* GetCurl() { return Curl; }
+	void SetClient(FCurlHttpClient* InClient) { Client = InClient; }
+	void SetComplete(CURLcode Code);
+
+	TRefCountPtr<IHttpResponseMonitor> GetMonitor() final { return this; }
+
+	void Cancel() final;
+	void Wait() const final { CompleteEvent.Wait(); }
+	bool Poll() const final { return EnumHasAnyFlags(State.load(), ECurlHttpResponseState::Complete); }
+	bool IsCanceled() const { return EnumHasAnyFlags(State.load(std::memory_order_relaxed), ECurlHttpResponseState::Canceled); }
+
+	FAnsiStringView GetUri() const final { return Uri; }
+	EHttpMethod GetMethod() const final { return Method; }
+	int32 GetStatusCode() const final { return StatusCode; }
+	EHttpErrorCode GetErrorCode() const final { return ErrorCode; }
+	FAnsiStringView GetError() const final { return Error; }
+	const FHttpResponseStats& GetStats() const final;
+	TConstArrayView<FAnsiStringView> GetAllHeaders() const final { return HeaderViews; }
+
+	void AddRef() const final;
+	void Release() const final;
+
+private:
+	~FCurlHttpResponse() final;
+	void Destroy() final;
+
+	void ConditionallySetResponseStatus();
+	bool WriteHeader(FAnsiStringView Header);
+	bool WriteBody(FMemoryView Body);
+
+	static size_t CurlDebug(CURL* Curl, curl_infotype DebugInfoType, char* DebugInfo, size_t DebugInfoSize, void* Param);
+	static size_t CurlHeader(void* Ptr, size_t SizeInBlocks, size_t BlockSizeInBytes, void* Param);
+	static size_t CurlWrite(void* Ptr, size_t SizeInBlocks, size_t BlockSizeInBytes, void* Param);
+
+	static EHttpErrorCode ConvertErrorCode(CURLcode Code);
+
+	FAnsiStringBuilderBase Uri;
+	EHttpMethod Method;
+	bool bHasBody = false;
+	EHttpErrorCode ErrorCode = EHttpErrorCode::Unknown;
+	std::atomic<ECurlHttpResponseState> State = ECurlHttpResponseState::None;
+	std::atomic<int32> StatusCode = -1;
+	mutable std::atomic<uint32> ReferenceCount = 0;
+
+	CURL* Curl; // null when complete
+	IHttpReceiver* Receiver;
+	FCurlHttpClient* Client = nullptr;
+	TArray<FAnsiStringView, TInlineAllocator<32>> HeaderViews;
+	TArray<int32, TInlineAllocator<32>> HeaderLengths;
+	TAnsiStringBuilder<4096> Headers;
+	FAnsiStringView Error;
+	ANSICHAR ErrorBuffer[CURL_ERROR_SIZE]{};
+	FHttpResponseStats Stats;
+	mutable FManualResetEvent CompleteEvent;
+
+	enum class EHttpReceiverFunction : uint8
+	{
+		None = 0,
+		OnCreate,
+		OnHeaders,
+		OnBody,
+		OnComplete,
+	};
+	inline static thread_local EHttpReceiverFunction ReceiverFunction;
+};
+
+///////////////////////////////////////////////////////////////////////////////////////////////////////////////////////
+
+FCurlHttpManager::FCurlHttpManager()
+{
+	// User-Agent: UnrealEngine/X.Y.Z-<CL> (<Platform>; <Config> <TargetType>; <BranchName>) <AppName>[CommandletName] (<ProjectName>)
+	const FEngineVersion& Version = FEngineVersion::Current();
+	UserAgent << ANSITEXTVIEW("User-Agent: UnrealEngine/")
+		<< Version.GetMajor() << '.' << Version.GetMinor() << '.' << Version.GetPatch() << '-' << Version.GetChangelist()
+		<< ANSITEXTVIEW(" (") << FPlatformProperties::PlatformName()
+		<< ANSITEXTVIEW("; ") << LexToString(FApp::GetBuildConfiguration()) << ' ' << LexToString(FApp::GetBuildTargetType())
+		<< ANSITEXTVIEW("; ") << FApp::GetBranchName()
+		<< ANSITEXTVIEW(") ") << FApp::GetName();
+
+	if (IsRunningCommandlet())
+	{
+		FString CommandletName;
+		FParse::Value(FCommandLine::Get(), TEXT("Run="), CommandletName);
+		CommandletName.ToLowerInline();
+		UserAgent << '[' << CommandletName << ']';
+	}
+
+	if (FApp::HasProjectName() && FApp::GetName() != FApp::GetProjectName())
+	{
+		UserAgent << ANSITEXTVIEW(" (") << FApp::GetProjectName() << ')';
+	}
+	UserAgent.ToString();
+}
+
+THttpUniquePtr<IHttpConnectionPool> FCurlHttpManager::CreateConnectionPool(const FHttpConnectionPoolParams& Params)
+{
+	return THttpUniquePtr<IHttpConnectionPool>(new FCurlHttpConnectionPool(*this, Params));
+}
+
+void FCurlHttpManager::SetDefaultOptions(CURL* Curl, FCurlHttpHeaders& Headers)
+{
+	curl_easy_setopt(Curl, CURLOPT_NOSIGNAL, 1L);
+	curl_easy_setopt(Curl, CURLOPT_USERAGENT, UserAgent.GetData());
+
+	Headers.AddHeader(*WriteToAnsiString<64>(ANSITEXTVIEW("UE-IsBuildMachine: "), GIsBuildMachine));
+	Headers.AddHeader(*WriteToAnsiString<64>(ANSITEXTVIEW("UE-Session: "), SessionId));
+	Headers.AddHeader(*WriteToAnsiString<32>(ANSITEXTVIEW("UE-Request: "), RequestId.fetch_add(1, std::memory_order_relaxed)));
+
+	// Remove the Expect: 100-Continue header that curl adds by default because it adds latency.
+	Headers.AddHeader("Expect:");
+}
+
+///////////////////////////////////////////////////////////////////////////////////////////////////////////////////////
+
+FCurlHttpConnectionPool::FCurlHttpConnectionPool(FCurlHttpManager& InManager, const FHttpConnectionPoolParams& Params)
+	: Manager(InManager)
+	, CurlShare(curl_share_init())
+	, CurlMulti(curl_multi_init())
+	, bThreadStarting(false)
+	, bThreadStopping(!Params.bAllowAsync || !FPlatformProcess::SupportsMultithreading())
+{
+	curl_share_setopt(CurlShare, CURLSHOPT_USERDATA, this);
+	curl_share_setopt(CurlShare, CURLSHOPT_LOCKFUNC, CurlLock);
+	curl_share_setopt(CurlShare, CURLSHOPT_UNLOCKFUNC, CurlUnlock);
+	curl_share_setopt(CurlShare, CURLSHOPT_SHARE, CURL_LOCK_DATA_DNS);
+	curl_share_setopt(CurlShare, CURLSHOPT_SHARE, CURL_LOCK_DATA_SSL_SESSION);
+
+	checkf(!Params.MaxConnections || !Params.MinConnections || Params.MinConnections <= Params.MaxConnections,
+		TEXT("MinConnections (%u) exceeds MaxConnections (%u)"), Params.MinConnections, Params.MaxConnections);
+	constexpr uint32 DefaultConnectionCount = 8;
+	const uint32 MaxConnections = Params.MaxConnections ? Params.MaxConnections : FMath::Max(Params.MinConnections, DefaultConnectionCount);
+	const uint32 MinConnections = Params.MinConnections ? Params.MinConnections : FMath::Min(Params.MaxConnections, DefaultConnectionCount);
+	curl_multi_setopt(CurlMulti, CURLMOPT_MAXCONNECTS, MinConnections);
+	curl_multi_setopt(CurlMulti, CURLMOPT_MAX_TOTAL_CONNECTIONS, MaxConnections);
+	if (Params.MaxRequestsPerConnection)
+	{
+		curl_multi_setopt(CurlMulti, CURLMOPT_MAX_CONCURRENT_STREAMS, (long)FMath::Min(Params.MaxRequestsPerConnection, (uint32)MAX_int32));
+	}
+	curl_multi_setopt(CurlMulti, CURLMOPT_PIPELINING, CURLPIPE_MULTIPLEX);
+}
+
+FCurlHttpConnectionPool::~FCurlHttpConnectionPool()
+{
+	checkf(ClientCount.load(std::memory_order_relaxed) == 0,
+		TEXT("Connection pool has %u clients that were not deleted."), ClientCount.load());
+
+	bThreadStopping.store(true, std::memory_order_relaxed);
+	AssertMultiCodeOk(curl_multi_wakeup(CurlMulti));
+	if (Thread.IsJoinable())
+	{
+		Thread.Join();
+	}
+
+	check(ThreadCommands.IsEmpty());
+
+	AssertMultiCodeOk(curl_multi_cleanup(CurlMulti));
+	AssertShareCodeOk(curl_share_cleanup(CurlShare));
+}
+
+THttpUniquePtr<IHttpClient> FCurlHttpConnectionPool::CreateClient(const FHttpClientParams& Params)
+{
+	ClientCount.fetch_add(1, std::memory_order_relaxed);
+	return THttpUniquePtr<IHttpClient>(new FCurlHttpClient(*this, Params));
+}
+
+void FCurlHttpConnectionPool::DeleteClient(FCurlHttpClient* Client)
+{
+	ClientCount.fetch_sub(1, std::memory_order_relaxed);
+}
+
+void FCurlHttpConnectionPool::SetDefaultOptions(CURL* Curl, FCurlHttpHeaders& Headers) const
+{
+	Manager.SetDefaultOptions(Curl, Headers);
+
+	curl_easy_setopt(Curl, CURLOPT_SHARE, CurlShare);
+}
+
+bool FCurlHttpConnectionPool::BeginAsyncRequest(FCurlHttpResponse* Response)
+{
+	if (bThreadStopping.load(std::memory_order_relaxed))
+	{
+		return false;
+	}
+	if (ThreadCommands.ProduceItem(FThreadCommand{Response, EThreadCommandType::Begin}) == EConsumeAllMpmcQueueResult::WasEmpty)
+	{
+		AssertMultiCodeOk(curl_multi_wakeup(CurlMulti));
+	}
+	if (!bThreadStarting.load(std::memory_order_relaxed) && !bThreadStarting.exchange(true, std::memory_order_relaxed))
+	{
+		Thread = FThread(TEXT("HttpConnectionPool"), [this] { ThreadLoop(); }, 128 * 1024);
+	}
+	return true;
+}
+
+void FCurlHttpConnectionPool::CancelAsyncRequest(FCurlHttpResponse* Response)
+{
+	if (ThreadCommands.ProduceItem(FThreadCommand{Response, EThreadCommandType::Cancel}) == EConsumeAllMpmcQueueResult::WasEmpty)
+	{
+		AssertMultiCodeOk(curl_multi_wakeup(CurlMulti));
+	}
+}
+
+void FCurlHttpConnectionPool::ThreadLoop()
+{
+	while (!ThreadCommands.IsEmpty() || !bThreadStopping.load(std::memory_order_relaxed))
+	{
+		ThreadCommands.ConsumeAllFifo([this](FThreadCommand Command)
+		{
+			if (CURL* Curl = Command.Response->GetCurl())
+			{
+				switch (Command.Type)
+				{
+				case EThreadCommandType::Begin:
+					AssertMultiCodeOk(curl_multi_add_handle(CurlMulti, Curl));
+					break;
+				case EThreadCommandType::Cancel:
+					CompleteRequest(Curl, nullptr);
+					break;
+				}
+			}
+		});
+
+		int ActiveRequests = 0;
+		AssertMultiCodeOk(curl_multi_perform(CurlMulti, &ActiveRequests));
+
+		for (int MessagesInQueue = 0; CURLMsg* Message = curl_multi_info_read(CurlMulti, &MessagesInQueue);)
+		{
+			if (Message->msg == CURLMSG_DONE)
+			{
+				CompleteRequest(Message->easy_handle, &Message->data.result);
+			}
+		}
+
+		constexpr int WaitTimeMs = 60'000;
+		AssertMultiCodeOk(curl_multi_poll(CurlMulti, nullptr, 0, WaitTimeMs, nullptr));
+	}
+}
+
+void FCurlHttpConnectionPool::CompleteRequest(CURL* Curl, CURLcode* OptionalReturnCode)
+{
+	AssertMultiCodeOk(curl_multi_remove_handle(CurlMulti, Curl));
+	void* Request = nullptr;
+	curl_easy_getinfo(Curl, CURLINFO_PRIVATE, &Request);
+	((FCurlHttpRequest*)Request)->OnComplete(OptionalReturnCode ? *OptionalReturnCode : CURLE_OK);
+}
+
+void FCurlHttpConnectionPool::CurlLock(CURL* Curl, curl_lock_data Data, curl_lock_access Access, void* Param)
+{
+	FCurlHttpConnectionPool& ConnectionPool = *(FCurlHttpConnectionPool*)Param;
+	if (Access == CURL_LOCK_ACCESS_SHARED)
+	{
+		ConnectionPool.Locks[Data].ReadLock();
+	}
+	else
+	{
+		ConnectionPool.Locks[Data].WriteLock();
+		ConnectionPool.WriteLocked[Data] = true;
+	}
+}
+
+void FCurlHttpConnectionPool::CurlUnlock(CURL* Curl, curl_lock_data Data, void* Param)
+{
+	FCurlHttpConnectionPool& ConnectionPool = *(FCurlHttpConnectionPool*)Param;
+	if (!ConnectionPool.WriteLocked[Data])
+	{
+		ConnectionPool.Locks[Data].ReadUnlock();
+	}
+	else
+	{
+		ConnectionPool.WriteLocked[Data] = false;
+		ConnectionPool.Locks[Data].WriteUnlock();
+	}
+}
+
+inline void FCurlHttpConnectionPool::AssertShareCodeOk(const CURLSHcode Code)
+{
+	checkf(Code == CURLSHE_OK, TEXT("Error in curl_share operation: %hs"), curl_share_strerror(Code));
+}
+
+inline void FCurlHttpConnectionPool::AssertMultiCodeOk(const CURLMcode Code)
+{
+	checkf(Code == CURLM_OK, TEXT("Error in curl_multi operation: %hs"), curl_multi_strerror(Code));
+}
+
+///////////////////////////////////////////////////////////////////////////////////////////////////////////////////////
+
+FCurlHttpClient::FCurlHttpClient(FCurlHttpConnectionPool& InConnectionPool, const FHttpClientParams& InParams)
+	: ConnectionPool(InConnectionPool)
+	, Params(InParams)
+{
+	checkf(!Params.MaxRequests || !Params.MinRequests || Params.MinRequests <= Params.MaxRequests,
+		TEXT("MinRequests (%u) exceeds MaxRequests (%u)"), Params.MinRequests, Params.MaxRequests);
+	Params.MaxRequests = Params.MaxRequests ? Params.MaxRequests : FMath::Max(Params.MinRequests, 256u);
+	Params.MinRequests = Params.MinRequests ? Params.MinRequests : FMath::Min(Params.MaxRequests, 16u);
+}
+
+FCurlHttpClient::~FCurlHttpClient()
+{
+	while (FCurlHttpRequest* Request = RequestPool.Pop())
+	{
+		delete Request;
+		RequestCount.fetch_sub(1, std::memory_order_relaxed);
+		RequestPoolCount.fetch_sub(1, std::memory_order_relaxed);
+	}
+
+	checkf(RequestCount.load(std::memory_order_relaxed) == 0,
+		TEXT("Client has %u requests that were not destroyed by IHttpRequest::Destroy()."), RequestCount.load());
+
+	ConnectionPool.DeleteClient(this);
+}
+
+THttpUniquePtr<IHttpRequest> FCurlHttpClient::TryCreateRequest(const FHttpRequestParams& RequestParams)
+{
+	if (IHttpRequest* Request = RequestPool.Pop())
+	{
+		RequestPoolCount.fetch_sub(1, std::memory_order_relaxed);
+		return THttpUniquePtr<IHttpRequest>(Request);
+	}
+
+	if (RequestParams.bIgnoreMaxRequests)
+	{
+		RequestCount.fetch_add(1, std::memory_order_relaxed);
+	}
+	else if (!CurlHttp::Private::TryIncrement(RequestCount, Params.MaxRequests))
+	{
+		return nullptr;
+	}
+
+	return THttpUniquePtr<IHttpRequest>(new FCurlHttpRequest(*this));
+}
+
+void FCurlHttpClient::DeleteRequest(FCurlHttpRequest* Request)
+{
+	if (CurlHttp::Private::TryIncrement(RequestPoolCount, Params.MinRequests))
+	{
+		Request->Reset();
+		RequestPool.Push(Request);
+	}
+	else
+	{
+		delete Request;
+		RequestCount.fetch_sub(1, std::memory_order_relaxed);
+	}
+
+	if (Params.OnDestroyRequest)
+	{
+		Params.OnDestroyRequest();
+	}
+}
+
+void FCurlHttpClient::SetDefaultOptions(CURL* Curl, FCurlHttpHeaders& Headers) const
+{
+	ConnectionPool.SetDefaultOptions(Curl, Headers);
+
+	curl_easy_setopt(Curl, CURLOPT_HTTP_VERSION, ConvertVersion(Params.Version));
+
+	curl_easy_setopt(Curl, CURLOPT_DNS_CACHE_TIMEOUT, (long)FMath::Min(Params.DnsCacheTimeout, (uint32)MAX_int32));
+	curl_easy_setopt(Curl, CURLOPT_CONNECTTIMEOUT_MS, long(Params.ConnectTimeout));
+	curl_easy_setopt(Curl, CURLOPT_LOW_SPEED_LIMIT, long(Params.LowSpeedLimit));
+	curl_easy_setopt(Curl, CURLOPT_LOW_SPEED_TIME, long(Params.LowSpeedTime));
+	curl_easy_setopt(Curl, CURLOPT_USE_SSL, ConvertTlsLevel(Params.TlsLevel));
+	curl_easy_setopt(Curl, CURLOPT_FOLLOWLOCATION, long(Params.bFollowRedirects));
+	curl_easy_setopt(Curl, CURLOPT_POSTREDIR, long(
+		(Params.bFollow301Post ? CURL_REDIR_POST_301 : 0) |
+		(Params.bFollow302Post ? CURL_REDIR_POST_302 : 0) |
+		(Params.bFollow303Post ? CURL_REDIR_POST_303 : 0)));
+	if (Params.bBypassProxy)
+	{
+		curl_easy_setopt(Curl, CURLOPT_NOPROXY, "*");
+	}
+	curl_easy_setopt(Curl, CURLOPT_VERBOSE, long(Params.bVerbose));
+}
+
+long FCurlHttpClient::ConvertVersion(const EHttpVersion Version)
+{
+	switch (Version)
+	{
+	case EHttpVersion::None:   return CURL_HTTP_VERSION_NONE;
+	case EHttpVersion::V1_0:   return CURL_HTTP_VERSION_1_0;
+	case EHttpVersion::V1_1:   return CURL_HTTP_VERSION_1_1;
+	case EHttpVersion::V2:     return CURL_HTTP_VERSION_2_0;
+	case EHttpVersion::V2Only: return CURL_HTTP_VERSION_2_PRIOR_KNOWLEDGE;
+	default: checkNoEntry();   return CURL_HTTP_VERSION_NONE;
+	}
+}
+
+long FCurlHttpClient::ConvertTlsLevel(const EHttpTlsLevel Level)
+{
+	switch (Level)
+	{
+	case EHttpTlsLevel::None: return CURLUSESSL_NONE;
+	case EHttpTlsLevel::Try:  return CURLUSESSL_TRY;
+	case EHttpTlsLevel::All:  return CURLUSESSL_ALL;
+	default: checkNoEntry();  return CURLUSESSL_NONE;
+	}
+}
+
+inline bool FCurlHttpClient::BeginAsyncRequest(FCurlHttpResponse* Response)
+{
+	Response->SetClient(this);
+	return ConnectionPool.BeginAsyncRequest(Response);
+}
+
+inline void FCurlHttpClient::CancelAsyncRequest(FCurlHttpResponse* Response)
+{
+	ConnectionPool.CancelAsyncRequest(Response);
+}
+
+///////////////////////////////////////////////////////////////////////////////////////////////////////////////////////
+
+FCurlHttpRequest::FCurlHttpRequest(FCurlHttpClient& InClient)
+	: Client(InClient)
+	, Curl(curl_easy_init())
+{
+	SetDefaultOptions();
+}
+
+FCurlHttpRequest::~FCurlHttpRequest()
+{
+	curl_easy_cleanup(Curl);
+}
+
+void FCurlHttpRequest::Reset()
+{
+	CheckIdle(TEXT("Reset"));
+	curl_easy_reset(Curl);
+	Headers.Reset();
+	Uri.Reset();
+	UnixSocketPath.Empty();
+	Body.Reset();
+	SetDefaultOptions();
+}
+
+void FCurlHttpRequest::SetDefaultOptions()
+{
+	Method = EHttpMethod::Get;
+
+	Client.SetDefaultOptions(Curl, Headers);
+
+	curl_easy_setopt(Curl, CURLOPT_PRIVATE, this);
+	curl_easy_setopt(Curl, CURLOPT_READDATA, this);
+	curl_easy_setopt(Curl, CURLOPT_READFUNCTION, CurlRead);
+	curl_easy_setopt(Curl, CURLOPT_SEEKDATA, this);
+	curl_easy_setopt(Curl, CURLOPT_SEEKFUNCTION, CurlSeek);
+#if WITH_SSL
+	curl_easy_setopt(Curl, CURLOPT_SSL_CTX_DATA, this);
+	curl_easy_setopt(Curl, CURLOPT_SSL_CTX_FUNCTION, CurlSslContext);
+#endif
+}
+
+void FCurlHttpRequest::CheckIdle(const TCHAR* FunctionName)
+{
+	checkf(!Response.load(std::memory_order_relaxed),
+		TEXT("%.*hs %.*hs: %s requires the request to be idle."),
+		LexToString(Method).Len(), LexToString(Method).GetData(), Uri.Len(), Uri.GetData(), FunctionName);
+}
+
+void FCurlHttpRequest::SetUri(const FAnsiStringView InUri)
+{
+	CheckIdle(TEXT("SetUri"));
+	Uri.Reset();
+	Uri.Append(InUri);
+	curl_easy_setopt(Curl, CURLOPT_URL, *Uri);
+}
+
+void FCurlHttpRequest::SetUnixSocketPath(const FAnsiStringView InSocketPath)
+{
+	CheckIdle(TEXT("SetUnixSocketPath"));
+	UnixSocketPath.Empty();
+	UnixSocketPath.Append(InSocketPath);
+
+	if (UnixSocketPath.Len() > 0)
+	{
+		curl_easy_setopt(Curl, CURLOPT_UNIX_SOCKET_PATH, *UnixSocketPath);
+	}
+	else
+	{
+		curl_easy_setopt(Curl, CURLOPT_UNIX_SOCKET_PATH, NULL);
+	}
+}
+
+void FCurlHttpRequest::SetMethod(const EHttpMethod InMethod)
+{
+	CheckIdle(TEXT("SetMethod"));
+	Method = InMethod;
+	switch (Method)
+	{
+	case EHttpMethod::Get:
+		curl_easy_setopt(Curl, CURLOPT_HTTPGET, 1L);
+		break;
+	case EHttpMethod::Put:
+		curl_easy_setopt(Curl, CURLOPT_UPLOAD, 1L);
+		break;
+	case EHttpMethod::Post:
+		curl_easy_setopt(Curl, CURLOPT_POST, 1L);
+		break;
+	case EHttpMethod::Head:
+		curl_easy_setopt(Curl, CURLOPT_NOBODY, 1L);
+		break;
+	case EHttpMethod::Delete:
+		curl_easy_setopt(Curl, CURLOPT_CUSTOMREQUEST, "DELETE");
+		break;
+	default:
+		checkNoEntry();
+		break;
+	}
+}
+
+void FCurlHttpRequest::SetBody(const FCompositeBuffer& InBody)
+{
+	CheckIdle(TEXT("SetBody"));
+	Body = InBody;
+}
+
+void FCurlHttpRequest::AddHeader(const FAnsiStringView Name, const FAnsiStringView Value)
+{
+	CheckIdle(TEXT("AddHeader"));
+	Headers.AddHeader(*WriteToAnsiString<256>(Name, ANSITEXTVIEW(": "), Value));
+}
+
+void FCurlHttpRequest::Send(IHttpReceiver* Receiver, THttpUniquePtr<IHttpResponse>& OutResponse)
+{
+	CheckIdle(TEXT("Send"));
+	if (CreateResponse(Receiver, OutResponse))
+	{
+		OnComplete(curl_easy_perform(Curl));
+	}
+}
+
+void FCurlHttpRequest::SendAsync(IHttpReceiver* Receiver, THttpUniquePtr<IHttpResponse>& OutResponse)
+{
+	CheckIdle(TEXT("SendAsync"));
+	if (FCurlHttpResponse* LocalResponse = CreateResponse(Receiver, OutResponse))
+	{
+		if (!Client.BeginAsyncRequest(LocalResponse))
+		{
+			OnComplete(curl_easy_perform(Curl));
+		}
+	}
+}
+
+FCurlHttpResponse* FCurlHttpRequest::CreateResponse(IHttpReceiver* Receiver, THttpUniquePtr<IHttpResponse>& OutResponse)
+{
+	BodyOffset = 0;
+
+	if (const uint64 BodySize = Body.GetSize(); BodySize > 0)
+	{
+		switch (Method)
+		{
+		case EHttpMethod::Put:
+			curl_easy_setopt(Curl, CURLOPT_INFILESIZE_LARGE, curl_off_t(BodySize));
+			break;
+		case EHttpMethod::Post:
+			curl_easy_setopt(Curl, CURLOPT_POSTFIELDSIZE_LARGE, curl_off_t(BodySize));
+			break;
+		default:
+			checkf(false, TEXT("%s"), *WriteToString<256>(TEXTVIEW("Method "), LexToString(Method),
+				TEXTVIEW(" must not include a body but a body of "), BodySize, TEXTVIEW(" bytes was provided.")));
+			break;
+		}
+	}
+
+	curl_easy_setopt(Curl, CURLOPT_HTTPHEADER, Headers.GetList());
+
+	FCurlHttpResponse* LocalResponse = new FCurlHttpResponse(Curl, Method, Uri, Receiver);
+	OutResponse = THttpUniquePtr<IHttpResponse>(LocalResponse);
+
+	if (LocalResponse->Create())
+	{
+		verify(!Response.exchange(LocalResponse, std::memory_order_release));
+		return LocalResponse;
+	}
+	LocalResponse->SetComplete(CURLE_OK);
+	return nullptr;
+}
+
+void FCurlHttpRequest::OnComplete(CURLcode Code)
+{
+	if (FCurlHttpResponse* LocalResponse = Response.exchange(nullptr, std::memory_order_acquire))
+	{
+		LocalResponse->SetComplete(Code);
+	}
+}
+
+size_t FCurlHttpRequest::CurlRead(void* Ptr, size_t SizeInBlocks, size_t BlockSizeInBytes, void* Param)
+{
+	FCurlHttpRequest* Request = (FCurlHttpRequest*)Param;
+	const size_t TotalSize = Request->Body.GetSize();
+	const size_t Offset = Request->BodyOffset;
+	const size_t Size = FMath::Min(TotalSize - Offset, SizeInBlocks * BlockSizeInBytes);
+	Request->Body.CopyTo(MakeMemoryView(Ptr, Size), Offset);
+	Request->BodyOffset += Size;
+	return Size;
+}
+
+size_t FCurlHttpRequest::CurlSeek(void* Param, curl_off_t Offset, int Origin)
+{
+	FCurlHttpRequest* Request = (FCurlHttpRequest*)Param;
+
+	const size_t TotalSize = Request->Body.GetSize();
+	size_t AbsoluteOffset = 0;
+
+	switch (Origin)
+	{
+	case SEEK_SET: AbsoluteOffset = Offset; break;
+	case SEEK_CUR: AbsoluteOffset = Request->BodyOffset + Offset; break;
+	case SEEK_END: AbsoluteOffset = TotalSize + Offset; break;
+	}
+
+	if (AbsoluteOffset >= TotalSize)
+	{
+		return CURL_SEEKFUNC_FAIL;
+	}
+
+	Request->BodyOffset = AbsoluteOffset;
+	return CURL_SEEKFUNC_OK;
+}
+
+#if WITH_SSL
+CURLcode FCurlHttpRequest::CurlSslContext(CURL* Curl, void* Context, void* Param)
+{
+	SSL_CTX* SslContext = (SSL_CTX*)Context;
+	FSslModule::Get().GetCertificateManager().AddCertificatesToSslContext(SslContext);
+	SSL_CTX_set_verify(SslContext, SSL_CTX_get_verify_mode(SslContext), SslCertVerify);
+	SSL_CTX_set_app_data(SslContext, Param);
+	return CURLE_OK;
+}
+
+int FCurlHttpRequest::SslCertVerify(int PreverifyOk, X509_STORE_CTX* Context)
+{
+	if (PreverifyOk == 1)
+	{
+		SSL* Ssl = (SSL*)X509_STORE_CTX_get_ex_data(Context, SSL_get_ex_data_X509_STORE_CTX_idx());
+		check(Ssl);
+
+		SSL_CTX* SslContext = SSL_get_SSL_CTX(Ssl);
+		check(SslContext);
+
+		FCurlHttpRequest* Request = (FCurlHttpRequest*)SSL_CTX_get_app_data(SslContext);
+		check(Request);
+
+		// Extract the domain from the URI.
+		FAnsiStringView Domain = Request->Uri;
+		if (const int32 SchemeIndex = String::FindFirst(Domain, ANSITEXTVIEW("://")); SchemeIndex != INDEX_NONE)
+		{
+			Domain.RightChopInline(SchemeIndex + ANSITEXTVIEW("://").Len());
+		}
+		if (const int32 SlashIndex = String::FindFirstChar(Domain, '/'); SlashIndex != INDEX_NONE)
+		{
+			Domain.LeftInline(SlashIndex);
+		}
+		if (const int32 AtIndex = String::FindFirstChar(Domain, '@'); AtIndex != INDEX_NONE)
+		{
+			Domain.RightChopInline(AtIndex + 1);
+		}
+		const auto RemovePort = [](FAnsiStringView& Authority)
+		{
+			if (const int32 ColonIndex = String::FindLastChar(Authority, ':'); ColonIndex != INDEX_NONE)
+			{
+				Authority.LeftInline(ColonIndex);
+			}
+		};
+		if (Domain.StartsWith('['))
+		{
+			if (const int32 LastBracketIndex = String::FindLastChar(Domain, ']'); LastBracketIndex != INDEX_NONE)
+			{
+				Domain.MidInline(1, LastBracketIndex - 1);
+			}
+			else
+			{
+				RemovePort(Domain);
+			}
+		}
+		else
+		{
+			RemovePort(Domain);
+		}
+
+		if (!FSslModule::Get().GetCertificateManager().VerifySslCertificates(Context, FString(Domain)))
+		{
+			PreverifyOk = 0;
+		}
+	}
+
+	return PreverifyOk;
+}
+#endif
+
+///////////////////////////////////////////////////////////////////////////////////////////////////////////////////////
+
+FCurlHttpResponse::FCurlHttpResponse(CURL* InCurl, EHttpMethod InMethod, FAnsiStringView InUri, IHttpReceiver* InReceiver)
+	: Method(InMethod)
+	, Curl(InCurl)
+	, Receiver(InReceiver)
+{
+	CaptureInheritedContext();
+
+	// Release() is called by Destroy().
+	AddRef();
+
+	Uri.Append(InUri);
+	checkf(Receiver, TEXT("Receiver must not be null for %s"), *WriteToString<256>(*this));
+
+	curl_easy_setopt(Curl, CURLOPT_ERRORBUFFER, ErrorBuffer);
+	curl_easy_setopt(Curl, CURLOPT_DEBUGDATA, this);
+	curl_easy_setopt(Curl, CURLOPT_DEBUGFUNCTION, CurlDebug);
+	curl_easy_setopt(Curl, CURLOPT_HEADERDATA, this);
+	curl_easy_setopt(Curl, CURLOPT_HEADERFUNCTION, CurlHeader);
+	curl_easy_setopt(Curl, CURLOPT_WRITEDATA, this);
+	curl_easy_setopt(Curl, CURLOPT_WRITEFUNCTION, CurlWrite);
+}
+
+bool FCurlHttpResponse::Create()
+{
+	for (TGuardValue GuardReceiverFunction(ReceiverFunction, EHttpReceiverFunction::OnCreate);;)
+	{
+		IHttpReceiver* NewReceiver = Receiver->OnCreate(*this);
+		if (Receiver == NewReceiver)
+		{
+			break;
+		}
+		checkf(NewReceiver, TEXT("Receiver must not be null for %s"), *WriteToString<256>(*this));
+		Receiver = NewReceiver;
+	}
+	return !IsCanceled();
+}
+
+FCurlHttpResponse::~FCurlHttpResponse()
+{
+	checkf(!Curl, TEXT("Response must be completed before it is destroyed for %s"), *WriteToString<256>(*this));
+}
+
+void FCurlHttpResponse::Destroy()
+{
+	if (Curl)
+	{
+		Cancel();
+	}
+
+	// AddRef() is called by the constructor.
+	Release();
+}
+
+void FCurlHttpResponse::ConditionallySetResponseStatus()
+{
+	if (StatusCode < 0)
+	{
+		long ResponseCode;
+		curl_easy_getinfo(Curl, CURLINFO_RESPONSE_CODE, &ResponseCode);
+		if (ResponseCode > 0)
+		{
+			StatusCode = int32(ResponseCode);
+		}
+	}
+}
+
+bool FCurlHttpResponse::WriteHeader(FAnsiStringView Header)
+{
+	if (bHasBody)
+	{
+		return !EnumHasAnyFlags(State.load(std::memory_order_relaxed), ECurlHttpResponseState::Canceled);
+	}
+
+	FInheritedContextScope InheritedContextScope = RestoreInheritedContext();
+
+	// Reset headers between responses and keep only the last.
+	if (!HeaderViews.IsEmpty())
+	{
+		HeaderLengths.Reset();
+		HeaderViews.Reset();
+		Headers.Reset();
+		StatusCode = -1;
+	}
+
+	ConditionallySetResponseStatus();
+
+	if (Header == ANSITEXTVIEW("\r\n"))
+	{
+		const ANSICHAR* Data = Headers.GetData();
+		for (const int32 Len : HeaderLengths)
+		{
+			HeaderViews.Add(MakeStringView(Data, Len - 2));
+			Data += Len;
+		}
+		for (TGuardValue GuardReceiverFunction(ReceiverFunction, EHttpReceiverFunction::OnHeaders);;)
+		{
+			IHttpReceiver* NewReceiver = Receiver->OnHeaders(*this);
+			if (Receiver == NewReceiver)
+			{
+				break;
+			}
+			checkf(NewReceiver, TEXT("Receiver must not be null for %s"), *WriteToString<256>(*this));
+			Receiver = NewReceiver;
+		}
+	}
+	else
+	{
+		HeaderLengths.Add(Header.Len());
+		Headers.Append(Header);
+	}
+
+	return !IsCanceled();
+}
+
+bool FCurlHttpResponse::WriteBody(FMemoryView Body)
+{
+	FInheritedContextScope InheritedContextScope = RestoreInheritedContext();
+
+	ConditionallySetResponseStatus();
+
+	bHasBody = true;
+
+	while (!Body.IsEmpty() && !IsCanceled())
+	{
+		const FMemoryView InitialBody = Body;
+		IHttpReceiver* NewReceiver;
+		{
+			TGuardValue GuardReceiverFunction(ReceiverFunction, EHttpReceiverFunction::OnBody);
+			NewReceiver = Receiver->OnBody(*this, Body);
+		}
+		if (Receiver == NewReceiver)
+		{
+			checkf(Body.IsEmpty() || Body == InitialBody,
+				TEXT("Receiver must consume the entire body for %s"), *WriteToString<256>(*this));
+			break;
+		}
+		checkf(Body.IsEmpty() || Body.GetDataEnd() == InitialBody.GetDataEnd(),
+			TEXT("Receiver must consume the body sequentially for %s"), *WriteToString<256>(*this));
+		checkf(NewReceiver, TEXT("Receiver must not be null for %s"), *WriteToString<256>(*this));
+		Receiver = NewReceiver;
+	}
+
+	return !IsCanceled();
+}
+
+void FCurlHttpResponse::SetComplete(CURLcode Code)
+{
+	Error = ErrorBuffer;
+	if (IsCanceled())
+	{
+		ErrorCode = EHttpErrorCode::Canceled;
+		Error = ANSITEXTVIEW("Canceled using IHttpResponse::Cancel()");
+		StatusCode = 0;
+	}
+	else if (Code == CURLE_OK)
+	{
+		ErrorCode = EHttpErrorCode::None;
+		ConditionallySetResponseStatus();
+	}
+	else
+	{
+		ErrorCode = ConvertErrorCode(Code);
+		if (Error.IsEmpty())
+		{
+			Error = curl_easy_strerror(Code);
+		}
+		StatusCode = 0;
+	}
+
+	const auto GetSizeInfo = [this](CURLINFO Info)
+	{
+		curl_off_t Value;
+		curl_easy_getinfo(Curl, Info, &Value);
+		return uint64(Value);
+	};
+	const auto GetDoubleInfo = [this](CURLINFO Info)
+	{
+		double Value;
+		curl_easy_getinfo(Curl, Info, &Value);
+		return Value;
+	};
+	Stats.SendSize = GetSizeInfo(CURLINFO_SIZE_UPLOAD_T);
+	Stats.RecvSize = GetSizeInfo(CURLINFO_SIZE_DOWNLOAD_T);
+	Stats.SendRate = GetSizeInfo(CURLINFO_SPEED_UPLOAD_T);
+	Stats.RecvRate = GetSizeInfo(CURLINFO_SPEED_DOWNLOAD_T);
+	Stats.NameResolveTime = GetDoubleInfo(CURLINFO_NAMELOOKUP_TIME);
+	Stats.ConnectTime = GetDoubleInfo(CURLINFO_CONNECT_TIME);
+	Stats.TlsConnectTime = GetDoubleInfo(CURLINFO_APPCONNECT_TIME);
+	Stats.PreTransferTime = GetDoubleInfo(CURLINFO_PRETRANSFER_TIME);
+	Stats.StartTransferTime = GetDoubleInfo(CURLINFO_STARTTRANSFER_TIME);
+	Stats.TotalTime = GetDoubleInfo(CURLINFO_TOTAL_TIME);
+
+	Curl = nullptr;
+	Client = nullptr;
+
+	// Hold a reference to safely access Receiver, State, CompleteEvent, and FInheritedContextBase.
+	TRefCountPtr<FCurlHttpResponse> Self(this);
+	FInheritedContextScope InheritedContextScope = RestoreInheritedContext();
+
+	for (TGuardValue GuardReceiverFunction(ReceiverFunction, EHttpReceiverFunction::OnComplete);;)
+	{
+		IHttpReceiver* NewReceiver = Receiver->OnComplete(*this);
+		if (Receiver == NewReceiver || !NewReceiver)
+		{
+			break;
+		}
+		Receiver = NewReceiver;
+	}
+
+	CurlHttp::Private::AtomicFetchOr(State, ECurlHttpResponseState::Complete);
+	CompleteEvent.Notify();
+
+	if (UE_LOG_ACTIVE(LogHttp, Verbose))
+	{
+		TStringBuilder<80> StatsText;
+		if (Stats.SendSize)
+		{
+			StatsText << TEXTVIEW("sent ") << Stats.SendSize << TEXTVIEW(" bytes, ");
+		}
+		if (Stats.RecvSize)
+		{
+			StatsText << TEXTVIEW("received ") << Stats.RecvSize << TEXTVIEW(" bytes, ");
+		}
+		StatsText.Appendf(TEXT("%.3f seconds"), Stats.TotalTime);
+		UE_LOG(LogHttp, Verbose, TEXT("%s (%s)"), *WriteToString<256>(*this), *StatsText);
+	}
+
+	// DO NOT ACCESS THIS AGAIN PAST THIS POINT!
+}
+
+void FCurlHttpResponse::Cancel()
+{
+	const EHttpReceiverFunction LocalReceiverFunction = ReceiverFunction;
+	checkf(LocalReceiverFunction != EHttpReceiverFunction::OnComplete,
+		TEXT("Cancel() must not be called from within OnComplete for %s"), *WriteToString<256>(*this));
+
+	// Cancel only once and only if not completed.
+	const auto NotCompleteOrCanceled = [](ECurlHttpResponseState CheckState) -> bool
+	{
+		return !EnumHasAnyFlags(CheckState, ECurlHttpResponseState::Complete | ECurlHttpResponseState::Canceled);
+	};
+	if (CurlHttp::Private::AtomicFetchOrIf(State, ECurlHttpResponseState::Canceled, NotCompleteOrCanceled))
+	{
+		// Queue cancellation of an async request when called from outside of a receiver function.
+		// Cancellation within a receiver is handled by WriterHeader or WriteBody returning false,
+		// which causes an error in curl that causes the request to complete immediately.
+		if (Client && LocalReceiverFunction == EHttpReceiverFunction::None)
+		{
+			Client->CancelAsyncRequest(this);
+		}
+	}
+
+	// Wait for completion when called from outside of a receiver function.
+	// Waiting within a receiver function would prevent the response from completing.
+	if (LocalReceiverFunction == EHttpReceiverFunction::None)
+	{
+		CompleteEvent.Wait();
+	}
+}
+
+const FHttpResponseStats& FCurlHttpResponse::GetStats() const
+{
+	checkf(!Curl, TEXT("Stats accessed before they are available for %s"), *WriteToString<256>(*this));
+	return Stats;
+}
+
+size_t FCurlHttpResponse::CurlDebug(CURL* Curl, curl_infotype DebugInfoType, char* DebugInfo, size_t DebugInfoSize, void* Param)
+{
+	if (DebugInfoType != CURLINFO_TEXT && DebugInfoType != CURLINFO_HEADER_IN)
+	{
+		return 0;
+	}
+
+	FAnsiStringView DebugText(DebugInfo, DebugInfoSize);
+	DebugText.TrimStartAndEndInline();
+	if (DebugText.IsEmpty())
+	{
+		return 0;
+	}
+
+	FCurlHttpResponse* Response = (FCurlHttpResponse*)Param;
+	const FAnsiStringView Method = LexToString(Response->GetMethod());
+	const FAnsiStringView Uri = Response->GetUri();
+	switch (DebugInfoType)
+	{
+	case CURLINFO_TEXT:
+		UE_LOG(LogHttp, VeryVerbose, TEXT("%.*hs %.*hs (%p): [TEXT] %.*hs"),
+			Method.Len(), Method.GetData(), Uri.Len(), Uri.GetData(), Response, DebugText.Len(), DebugText.GetData());
+		break;
+	case CURLINFO_HEADER_IN:
+		UE_LOG(LogHttp, VeryVerbose, TEXT("%.*hs %.*hs (%p): [RESPONSE] %.*hs"),
+			Method.Len(), Method.GetData(), Uri.Len(), Uri.GetData(), Response, DebugText.Len(), DebugText.GetData());
+		break;
+	}
+	return 0;
+}
+
+size_t FCurlHttpResponse::CurlHeader(void* Ptr, size_t SizeInBlocks, size_t BlockSizeInBytes, void* Param)
+{
+	const int32 HeaderSize = IntCastChecked<int32>(SizeInBlocks * BlockSizeInBytes);
+	return ((FCurlHttpResponse*)Param)->WriteHeader(MakeStringView((const ANSICHAR*)Ptr, HeaderSize)) ? HeaderSize : 0;
+}
+
+size_t FCurlHttpResponse::CurlWrite(void* Ptr, size_t SizeInBlocks, size_t BlockSizeInBytes, void* Param)
+{
+	const size_t Size = SizeInBlocks * BlockSizeInBytes;
+	return ((FCurlHttpResponse*)Param)->WriteBody(MakeMemoryView(Ptr, Size)) ? Size : 0;
+}
+
+EHttpErrorCode FCurlHttpResponse::ConvertErrorCode(CURLcode Code)
+{
+	switch (Code)
+	{
+	case CURLE_OK:
+		return EHttpErrorCode::None;
+	case CURLE_COULDNT_RESOLVE_HOST:
+		return EHttpErrorCode::ResolveHost;
+	case CURLE_COULDNT_CONNECT:
+		return EHttpErrorCode::Connect;
+	case CURLE_SSL_CONNECT_ERROR:
+		return EHttpErrorCode::TlsConnect;
+	case CURLE_PEER_FAILED_VERIFICATION:
+		return EHttpErrorCode::TlsPeerVerification;
+	case CURLE_OPERATION_TIMEDOUT:
+		return EHttpErrorCode::TimedOut;
+	default:
+		return EHttpErrorCode::Unknown;
+	}
+}
+
+void FCurlHttpResponse::AddRef() const
+{
+	ReferenceCount.fetch_add(1, std::memory_order_relaxed);
+}
+
+void FCurlHttpResponse::Release() const
+{
+	if (ReferenceCount.fetch_sub(1, std::memory_order_acq_rel) == 1)
+	{
+		delete this;
+	}
+}
+
+///////////////////////////////////////////////////////////////////////////////////////////////////////////////////////
+
+IHttpManager& IHttpManager::Get()
+{
+	static FCurlHttpManager CurlHttpManager;
+	return CurlHttpManager;
+}
+
+} // UE
diff --git a/Engine/Source/Developer/DevHttp/Private/HttpClient.cpp b/Engine/Source/Developer/DevHttp/Private/HttpClient.cpp
new file mode 100644
--- /dev/null
+++ b/Engine/Source/Developer/DevHttp/Private/HttpClient.cpp
@@ -0,0 +1,437 @@
+// Copyright Epic Games, Inc. All Rights Reserved.
+
+#include "Http/HttpClient.h"
+
+#include "Async/ManualResetEvent.h"
+#include "Containers/LockFreeList.h"
+#include "Memory/MemoryView.h"
+#include "Misc/AsciiSet.h"
+#include "ProfilingDebugging/CpuProfilerTrace.h"
+#include "String/Find.h"
+
+namespace UE
+{
+
+FAnsiStringView LexToString(const EHttpMethod Method)
+{
+	switch (Method)
+	{
+	case EHttpMethod::Get:    return ANSITEXTVIEW("GET");
+	case EHttpMethod::Put:    return ANSITEXTVIEW("PUT");
+	case EHttpMethod::Post:   return ANSITEXTVIEW("POST");
+	case EHttpMethod::Head:   return ANSITEXTVIEW("HEAD");
+	case EHttpMethod::Delete: return ANSITEXTVIEW("DELETE");
+	default:                  return ANSITEXTVIEW("UNKNOWN");
+	}
+}
+
+bool TryLexFromString(EHttpMethod& OutMethod, const FAnsiStringView View)
+{
+	if (View == ANSITEXTVIEW("GET"))
+	{
+		OutMethod = EHttpMethod::Get;
+	}
+	else if (View == ANSITEXTVIEW("PUT"))
+	{
+		OutMethod = EHttpMethod::Put;
+	}
+	else if (View == ANSITEXTVIEW("POST"))
+	{
+		OutMethod = EHttpMethod::Post;
+	}
+	else if (View == ANSITEXTVIEW("HEAD"))
+	{
+		OutMethod = EHttpMethod::Head;
+	}
+	else if (View == ANSITEXTVIEW("DELETE"))
+	{
+		OutMethod = EHttpMethod::Delete;
+	}
+	else
+	{
+		return false;
+	}
+	return true;
+}
+
+FAnsiStringView LexToString(const EHttpMediaType MediaType)
+{
+	switch (MediaType)
+	{
+	case EHttpMediaType::Any:              return ANSITEXTVIEW("*/*");
+	case EHttpMediaType::Binary:           return ANSITEXTVIEW("application/octet-stream");
+	case EHttpMediaType::Text:             return ANSITEXTVIEW("text/plain");
+	case EHttpMediaType::Json:             return ANSITEXTVIEW("application/json");
+	case EHttpMediaType::Yaml:             return ANSITEXTVIEW("text/yaml");
+	case EHttpMediaType::CbObject:         return ANSITEXTVIEW("application/x-ue-cb");
+	case EHttpMediaType::CbPackage:        return ANSITEXTVIEW("application/x-ue-cbpkg");
+	case EHttpMediaType::CbPackageOffer:   return ANSITEXTVIEW("application/x-ue-offer");
+	case EHttpMediaType::CompressedBinary: return ANSITEXTVIEW("application/x-ue-comp");
+	case EHttpMediaType::FormUrlEncoded:   return ANSITEXTVIEW("application/x-www-form-urlencoded");
+	default:                               return ANSITEXTVIEW("unknown");
+	}
+}
+
+bool TryLexFromString(EHttpMediaType& OutMediaType, const FAnsiStringView View)
+{
+	const int32 SlashIndex = String::FindFirstChar(View, '/');
+	if (SlashIndex == INDEX_NONE)
+	{
+		return false;
+	}
+
+	const FAnsiStringView Type = View.Left(SlashIndex);
+	const FAnsiStringView SubType = View.RightChop(SlashIndex + 1);
+	if (Type == ANSITEXTVIEW("application"))
+	{
+		if (SubType == ANSITEXTVIEW("octet-stream"))
+		{
+			OutMediaType = EHttpMediaType::Binary;
+		}
+		else if (SubType == ANSITEXTVIEW("json"))
+		{
+			OutMediaType = EHttpMediaType::Json;
+		}
+		else if (SubType == ANSITEXTVIEW("x-ue-cb"))
+		{
+			OutMediaType = EHttpMediaType::CbObject;
+		}
+		else if (SubType == ANSITEXTVIEW("x-ue-cbpkg"))
+		{
+			OutMediaType = EHttpMediaType::CbPackage;
+		}
+		else if (SubType == ANSITEXTVIEW("x-ue-offer"))
+		{
+			OutMediaType = EHttpMediaType::CbPackageOffer;
+		}
+		else if (SubType == ANSITEXTVIEW("x-ue-comp"))
+		{
+			OutMediaType = EHttpMediaType::CompressedBinary;
+		}
+		else if (SubType == ANSITEXTVIEW("x-www-form-urlencoded"))
+		{
+			OutMediaType = EHttpMediaType::FormUrlEncoded;
+		}
+		else
+		{
+			return false;
+		}
+	}
+	else if (Type == ANSITEXTVIEW("text"))
+	{
+		if (SubType == ANSITEXTVIEW("plain"))
+		{
+			OutMediaType = EHttpMediaType::Text;
+		}
+		else if (SubType == ANSITEXTVIEW("yaml"))
+		{
+			OutMediaType = EHttpMediaType::Yaml;
+		}
+		else
+		{
+			return false;
+		}
+	}
+	else if (Type == ANSITEXTVIEW("*") && SubType == ANSITEXTVIEW("*"))
+	{
+		OutMediaType = EHttpMediaType::Any;
+	}
+	else
+	{
+		return false;
+	}
+
+	return true;
+}
+
+static FAnsiStringView LexStatusCodeToString(int HttpCode)
+{
+	switch (HttpCode)
+	{
+		// 1xx Informational
+	case 100: return "Continue";
+	case 101: return "Switching Protocols";
+
+		// 2xx Success
+	case 200: return "OK";
+	case 201: return "Created";
+	case 202: return "Accepted";
+	case 204: return "No Content";
+	case 205: return "Reset Content";
+	case 206: return "Partial Content";
+
+		// 3xx Redirection
+	case 300: return "Multiple Choices";
+	case 301: return "Moved Permanently";
+	case 302: return "Found";
+	case 303: return "See Other";
+	case 304: return "Not Modified";
+	case 305: return "Use Proxy";
+	case 306: return "Switch Proxy";
+	case 307: return "Temporary Redirect";
+	case 308: return "Permanent Redirect";
+
+		// 4xx Client errors
+	case 400: return "Bad Request";
+	case 401: return "Unauthorized";
+	case 402: return "Payment Required";
+	case 403: return "Forbidden";
+	case 404: return "Not Found";
+	case 405: return "Method Not Allowed";
+	case 406: return "Not Acceptable";
+	case 407: return "Proxy Authentication Required";
+	case 408: return "Request Timeout";
+	case 409: return "Conflict";
+	case 410: return "Gone";
+	case 411: return "Length Required";
+	case 412: return "Precondition Failed";
+	case 413: return "Payload Too Large";
+	case 414: return "URI Too Long";
+	case 415: return "Unsupported Media Type";
+	case 416: return "Range Not Satisfiable";
+	case 417: return "Expectation Failed";
+	case 418: return "I'm a teapot";
+	case 421: return "Misdirected Request";
+	case 422: return "Unprocessable Entity";
+	case 423: return "Locked";
+	case 424: return "Failed Dependency";
+	case 425: return "Too Early";
+	case 426: return "Upgrade Required";
+	case 428: return "Precondition Required";
+	case 429: return "Too Many Requests";
+	case 431: return "Request Header Fields Too Large";
+
+		// 5xx Server errors
+	case 500: return "Internal Server Error";
+	case 501: return "Not Implemented";
+	case 502: return "Bad Gateway";
+	case 503: return "Service Unavailable";
+	case 504: return "Gateway Timeout";
+	case 505: return "HTTP Version Not Supported";
+	case 506: return "Variant Also Negotiates";
+	case 507: return "Insufficient Storage";
+	case 508: return "Loop Detected";
+	case 510: return "Not Extended";
+	case 511: return "Network Authentication Required";
+
+	default: return "Unknown Result";
+	}
+}
+
+
+///////////////////////////////////////////////////////////////////////////////////////////////////////////////////////
+
+void IHttpRequest::SetContentType(const EHttpMediaType Type, const FAnsiStringView Param)
+{
+	TAnsiStringBuilder<64> Value;
+	Value << LexToString(Type);
+	if (!Param.IsEmpty())
+	{
+		Value << ANSITEXTVIEW("; ") << Param;
+	}
+	AddHeader(ANSITEXTVIEW("Content-Type"), Value);
+}
+
+void IHttpRequest::AddAcceptType(const EHttpMediaType Type, const float Weight)
+{
+	TAnsiStringBuilder<64> Value;
+	Value << LexToString(Type);
+	if (Weight != 1.0f)
+	{
+		Value.Appendf(";q=%.3f", Weight);
+	}
+	AddHeader(ANSITEXTVIEW("Accept"), Value);
+}
+
+///////////////////////////////////////////////////////////////////////////////////////////////////////////////////////
+
+FAnsiStringView IHttpResponse::GetHeader(const FAnsiStringView Name) const
+{
+	const int32 NameLen = Name.Len();
+	for (const FAnsiStringView& Header : GetAllHeaders())
+	{
+		if (Header.StartsWith(Name, ESearchCase::IgnoreCase) &&
+			Header.Len() > NameLen && Header.GetData()[NameLen] == ':')
+		{
+			return Header.RightChop(NameLen + 1).TrimStartAndEnd();
+		}
+	}
+	return {};
+}
+
+int32 IHttpResponse::GetHeaders(FAnsiStringView Name, TArrayView<FAnsiStringView> OutValues) const
+{
+	int32 MatchIndex = 0;
+	const int32 NameLen = Name.Len();
+	const int32 ValuesCount = OutValues.Num();
+	FAnsiStringView* Values = OutValues.GetData();
+	for (const FAnsiStringView& Header : GetAllHeaders())
+	{
+		if (Header.StartsWith(Name, ESearchCase::IgnoreCase) &&
+			Header.Len() > NameLen && Header.GetData()[NameLen] == ':')
+		{
+			if (ValuesCount > MatchIndex)
+			{
+				Values[MatchIndex] = Header.RightChop(NameLen + 1).TrimStartAndEnd();
+			}
+			++MatchIndex;
+		}
+	}
+	return MatchIndex;
+}
+
+EHttpMediaType IHttpResponse::GetContentType() const
+{
+	const FAnsiStringView ContentType = GetHeader(ANSITEXTVIEW("Content-Type"));
+	const FAnsiStringView ContentTypeNoParams = FAsciiSet::FindPrefixWithout(ContentType, "; \t");
+	EHttpMediaType MediaType = EHttpMediaType::Any;
+	TryLexFromString(MediaType, ContentTypeNoParams);
+	return MediaType;
+}
+
+FStringBuilderBase& operator<<(FStringBuilderBase& Builder, const IHttpResponse& Response)
+{
+	Builder << LexToString(Response.GetMethod()) << TEXT(' ') << Response.GetUri();
+
+	if (const int32 StatusCode = Response.GetStatusCode(); StatusCode > 0)
+	{
+		Builder << TEXTVIEW(" -> ") << LexStatusCodeToString(StatusCode) << " (" << StatusCode << ")";
+	}
+
+	if (const FAnsiStringView Error = Response.GetError(); !Error.IsEmpty())
+	{
+		Builder << TEXTVIEW(": ") << Error;
+	}
+
+	return Builder;
+}
+
+///////////////////////////////////////////////////////////////////////////////////////////////////////////////////////
+
+FHttpByteArrayReceiver::FHttpByteArrayReceiver(TArray64<uint8>& OutArray, IHttpReceiver* InNext)
+	: Array(OutArray)
+	, Next(InNext)
+{
+	Array.Reset();
+}
+
+IHttpReceiver* FHttpByteArrayReceiver::OnBody(IHttpResponse& Response, FMemoryView& Data)
+{
+	if (Array.IsEmpty())
+	{
+		constexpr int64 MaxReserveSize = 96 * 1024 * 1024;
+		if (const FAnsiStringView View = Response.GetHeader(ANSITEXTVIEW("Content-Length")); !View.IsEmpty())
+		{
+			constexpr int32 MaxStringLen = 16;
+			ANSICHAR String[MaxStringLen];
+			if (const int32 CopyLen = View.CopyString(String, MaxStringLen); CopyLen < MaxStringLen)
+			{
+				String[CopyLen] = '\0';
+				const int64 ContentLength = FCStringAnsi::Atoi64(String);
+				if (ContentLength > 0 && ContentLength <= MaxReserveSize)
+				{
+					Array.Reserve(ContentLength);
+				}
+			}
+		}
+	}
+	Array.Append(static_cast<const uint8*>(Data.GetData()), int64(Data.GetSize()));
+	return this;
+}
+
+///////////////////////////////////////////////////////////////////////////////////////////////////////////////////////
+
+struct Http::Private::FHttpRequestQueueData
+{
+	struct FWaiter
+	{
+		THttpUniquePtr<IHttpRequest> Request;
+		FManualResetEvent Event;
+	};
+
+	FHttpRequestQueueData(IHttpConnectionPool& ConnectionPool, const FHttpClientParams& ClientParams)
+	{
+		FHttpClientParams QueueParams = ClientParams;
+		QueueParams.OnDestroyRequest = [this, OnDestroyRequest = MoveTemp(QueueParams.OnDestroyRequest)]
+		{
+			if (OnDestroyRequest)
+			{
+				OnDestroyRequest();
+			}
+			if (!Waiters.IsEmpty())
+			{
+				if (THttpUniquePtr<IHttpRequest> Request = Client->TryCreateRequest({}))
+				{
+					if (FWaiter* Waiter = Waiters.Pop())
+					{
+						Waiter->Request = MoveTemp(Request);
+						Waiter->Event.Notify();
+					}
+				}
+			}
+		};
+		Client = ConnectionPool.CreateClient(QueueParams);
+	}
+
+	THttpUniquePtr<IHttpRequest> CreateRequest(const FHttpRequestParams& Params)
+	{
+		if (Params.bIgnoreMaxRequests)
+		{
+			THttpUniquePtr<IHttpRequest> Request = Client->TryCreateRequest(Params);
+			checkf(Request, TEXT("IHttpClient::TryCreateRequest returned null in spite of bIgnoreMaxRequests."));
+			return Request;
+		}
+
+		while (THttpUniquePtr<IHttpRequest> Request = Client->TryCreateRequest(Params))
+		{
+			if (FWaiter* Waiter = Waiters.Pop())
+			{
+				Waiter->Request = MoveTemp(Request);
+				Waiter->Event.Notify();
+			}
+			else
+			{
+				return Request;
+			}
+		}
+
+		FWaiter LocalWaiter;
+		Waiters.Push(&LocalWaiter);
+
+		while (THttpUniquePtr<IHttpRequest> Request = Client->TryCreateRequest(Params))
+		{
+			if (FWaiter* Waiter = Waiters.Pop())
+			{
+				Waiter->Request = MoveTemp(Request);
+				Waiter->Event.Notify();
+			}
+			if (LocalWaiter.Event.IsNotified())
+			{
+				checkf(LocalWaiter.Request, TEXT("CreateRequest returning null after IsNotified()."));
+				return MoveTemp(LocalWaiter.Request);
+			}
+		}
+
+		TRACE_CPUPROFILER_EVENT_SCOPE(HttpRequestQueue_Wait);
+		LocalWaiter.Event.Wait();
+		checkf(LocalWaiter.Request, TEXT("CreateRequest returning null after Wait()."));
+		return MoveTemp(LocalWaiter.Request);
+	}
+
+	THttpUniquePtr<IHttpClient> Client;
+	TLockFreePointerListFIFO<FWaiter, 0> Waiters;
+};
+
+FHttpRequestQueue::FHttpRequestQueue(IHttpConnectionPool& ConnectionPool, const FHttpClientParams& ClientParams)
+	: Data(MakePimpl<Http::Private::FHttpRequestQueueData>(ConnectionPool, ClientParams))
+{
+}
+
+THttpUniquePtr<IHttpRequest> FHttpRequestQueue::CreateRequest(const FHttpRequestParams& Params)
+{
+	static_assert(sizeof(Params) == sizeof(bool), "CreateRequest only handles bIgnoreMaxRequests");
+	check(Data);
+	return Data->CreateRequest(Params);
+}
+
+} // UE
diff --git a/Engine/Source/Developer/DevHttp/Private/HttpHostBuilder.cpp b/Engine/Source/Developer/DevHttp/Private/HttpHostBuilder.cpp
new file mode 100644
--- /dev/null
+++ b/Engine/Source/Developer/DevHttp/Private/HttpHostBuilder.cpp
@@ -0,0 +1,343 @@
+// Copyright Epic Games, Inc. All Rights Reserved.
+
+#include "Http/HttpHostBuilder.h"
+
+#include "Algo/MinElement.h"
+#include "Async/ManualResetEvent.h"
+#include "Async/Mutex.h"
+#include "Async/UniqueLock.h"
+#include "Containers/StringFwd.h"
+#include "CoreTypes.h"
+#include "HAL/PlatformTime.h"
+#include "Http/HttpClient.h"
+#include "Logging/LogMacros.h"
+#include "Logging/StructuredLog.h"
+#include "Memory/MemoryFwd.h"
+#include "Misc/ScopeExit.h"
+#include "Serialization/CompactBinary.h"
+#include "String/Find.h"
+#include "String/ParseTokens.h"
+#include "Templates/Function.h"
+#include "Templates/UniquePtr.h"
+
+namespace UE
+{
+
+DEFINE_LOG_CATEGORY_STATIC(LogHttpHostBuilder, Display, All);
+
+class FHttpBenchmarkReceiver final : public IHttpReceiver
+{
+public:
+	FHttpBenchmarkReceiver(const FHttpBenchmarkReceiver&) = delete;
+	FHttpBenchmarkReceiver& operator=(const FHttpBenchmarkReceiver&) = delete;
+
+	using FOnComplete = TUniqueFunction<void(IHttpResponse& HttpResponse, FString& Host)>;
+
+	explicit FHttpBenchmarkReceiver(const FStringView InHostEntry, FOnComplete&& InOperationComplete, IHttpReceiver* InNext = nullptr)
+		: HostEntry(InHostEntry)
+		, OperationComplete(MoveTemp(InOperationComplete))
+		, Next(InNext)
+	{
+	}
+
+private:
+	IHttpReceiver* OnCreate(IHttpResponse& Response) final
+	{
+		return this;
+	}
+
+	IHttpReceiver* OnComplete(IHttpResponse& Response) final
+	{
+		OperationComplete(Response, HostEntry);
+
+		return Next;
+	}
+
+private:
+	FString HostEntry;
+	FOnComplete OperationComplete;
+	IHttpReceiver* Next;
+};
+
+class FHttpBenchmarkOperation final : public IHttpReceiver
+{
+public:
+	FHttpBenchmarkOperation(const FHttpBenchmarkOperation&) = delete;
+	FHttpBenchmarkOperation& operator=(const FHttpBenchmarkOperation&) = delete;
+
+	using FOnComplete = TUniqueFunction<void(IHttpResponse& HttpResponse, FString& Host)>;
+
+	explicit FHttpBenchmarkOperation(const FStringView InHostEntry, THttpUniquePtr<IHttpRequest>&& InRequest, FOnComplete&& InOperationComplete)
+		: HostEntry(InHostEntry)
+		, OperationComplete(MoveTemp(InOperationComplete))
+		, Request(MoveTemp(InRequest))
+		, Response(nullptr)
+	{
+		Request->SetUri(WriteToAnsiString<256>(HostEntry, ANSITEXTVIEW("/health/ready")));
+	}
+
+	void SendAsync()
+	{
+		Request->SendAsync(this, Response);
+	}
+
+	void Cancel()
+	{
+		Response->Cancel();
+	}
+	
+private:
+	IHttpReceiver* OnCreate(IHttpResponse& InResponse) final
+	{
+		return this;
+	}
+
+	IHttpReceiver* OnComplete(IHttpResponse& InResponse) final
+	{
+		OperationComplete(InResponse, HostEntry);
+
+		return nullptr;
+	}
+
+private:
+	FString HostEntry;
+	FOnComplete OperationComplete;
+
+	THttpUniquePtr<IHttpRequest> Request;
+	THttpUniquePtr<IHttpResponse> Response;
+};
+
+void FHttpHostBuilder::AddFromEndpoint(const FAnsiStringView HostUrl, const FAnsiStringView AccessToken)
+{
+	FHttpConnectionPoolParams ConnectionPoolParams;
+	const THttpUniquePtr<IHttpConnectionPool> LocalConnectionPool = IHttpManager::Get().CreateConnectionPool(ConnectionPoolParams);
+	FHttpClientParams DefaultClientParams;
+	// we want to keep this timeout fairly low in case the host is not reachable
+	DefaultClientParams.ConnectTimeout = 5000;
+	THttpUniquePtr<IHttpClient> LocalClient = LocalConnectionPool->CreateClient(DefaultClientParams);
+
+	FHttpRequestParams RequestParams;
+	// setting ignore max requests guarantees that CreateRequest actually creates a request no matter how many requests are in flight
+	RequestParams.bIgnoreMaxRequests = true;
+
+	TArray64<uint8> BodyArray;
+	FHttpByteArrayReceiver PeersReceiver(BodyArray);
+	
+	THttpUniquePtr<IHttpRequest> Request = LocalClient->TryCreateRequest(RequestParams);
+
+	if (!AccessToken.IsEmpty())
+	{
+		Request->AddHeader(ANSITEXTVIEW("Authorization"), WriteToAnsiString<1024>(AccessToken));
+	}
+	Request->AddHeader(ANSITEXTVIEW("Accept"), ANSITEXTVIEW("application/x-ue-cb"));
+	TAnsiStringBuilder<256> Uri(InPlace, HostUrl, ANSITEXTVIEW("/status/peers"));
+	Request->SetUri(Uri);
+	THttpUniquePtr<IHttpResponse> Response;
+	Request->Send(&PeersReceiver, Response);
+
+	if (Response->GetStatusCode() != 200)
+	{
+		UE_LOGFMT(LogHttpHostBuilder, Warning, "Unsuccessful attempt to fetch hosts by endpoint from host: '{Host}'. Status code was: {Code}", Uri, Response->GetStatusCode());
+		return;
+	}
+
+	FSharedBuffer SharedBuffer = MakeSharedBufferFromArray(MoveTemp(BodyArray));
+	FCbObject PeersObject = FCbObject(SharedBuffer);
+
+	for (FCbField Peer : PeersObject["peers"].AsArray())
+	{
+		FCbObject PeerObject = Peer.AsObject();
+		for (FCbField Endpoint : PeerObject["endpoints"])
+		{
+			HostCandidates.AddUnique(FString(Endpoint.AsString()));
+		}
+	}
+}
+
+FString FHttpHostBuilder::GetHostCandidatesString() const
+{
+	return FString::Join(HostCandidates, TEXT(", "));
+}
+
+void FHttpHostBuilder::AddFromString(const FAnsiStringView HostList)
+{
+	String::ParseTokens(HostList, ';', [this](FAnsiStringView Host)
+	{
+		this->HostCandidates.Emplace(Host);
+	}, String::EParseTokensOptions::SkipEmpty | String::EParseTokensOptions::Trim);
+}
+
+void FHttpHostBuilder::AddFromString(const FStringView HostList)
+{
+	String::ParseTokens(HostList, ';', [this](FStringView Host)
+	{
+		this->HostCandidates.Emplace(Host);
+	}, String::EParseTokensOptions::SkipEmpty | String::EParseTokensOptions::Trim);
+}
+
+bool FHttpHostBuilder::ResolveHost(double WarningTimeoutSeconds, double TimeoutSeconds, FAnsiStringBuilderBase& OutHost, double& OutLatency)
+{
+	if (HostCandidates.Num() == 0)
+	{
+		// no hosts have been added
+		return false;
+	}
+
+	if (HostCandidates.Num() == 1)
+	{
+		// if there is only one candidate we do not benchmark it and just return this straight away
+		OutHost.Reset();
+		OutHost.Append(HostCandidates[0]);
+
+		OutLatency = 0.0;
+		return true;
+	}
+
+	bool bHostFound = BenchmarkHostList(HostCandidates, WarningTimeoutSeconds, TimeoutSeconds, OutHost, OutLatency);
+
+	// if no valid host was found, set the host to the first candidate, even if this will likely not work to connect to its still better then no options
+	if (!bHostFound)
+	{
+		OutHost.Reset();
+		OutHost.Append(HostCandidates[0]);
+		OutLatency = 0.0;
+	}
+	return bHostFound;
+}
+
+bool FHttpHostBuilder::BenchmarkHostList(TConstArrayView<FString> InHostCandidates, double WarningTimeoutSeconds, double TimeoutSeconds, FAnsiStringBuilderBase& OutHost, double& OutLatency)
+{
+	const double StartTime = FPlatformTime::Seconds();
+
+	constexpr uint32 MaxTotalConnections = 8;
+	FHttpConnectionPoolParams ConnectionPoolParams;
+	ConnectionPoolParams.MaxConnections = MaxTotalConnections;
+	ConnectionPoolParams.MinConnections = MaxTotalConnections;
+	const THttpUniquePtr<IHttpConnectionPool> LocalConnectionPool = IHttpManager::Get().CreateConnectionPool(ConnectionPoolParams);
+	const FHttpClientParams DefaultClientParams;
+	THttpUniquePtr<IHttpClient> LocalClient = LocalConnectionPool->CreateClient(DefaultClientParams);
+
+	struct FLatencySortableHost
+	{
+		double Latency;
+		FString Host;
+	};
+	TArray<FLatencySortableHost> SortedHostList;
+	TArray<TUniquePtr<FHttpBenchmarkOperation>> Operations;
+
+	FMutex Mutex;
+	FManualResetEvent HostBenchmarkDone;
+
+	int32 FailedBenchmarkAttempts = 0;
+	bool bAllHostsFailed = false;
+	int32 CountOfBenchmarksToRun = InHostCandidates.Num();
+	for (const FString& HostCandidate : InHostCandidates)
+	{
+		// Benchmark each of the hosts
+		FHttpRequestParams RequestParams;
+		// setting ignore max requests guarantees that CreateRequest actually creates a request no matter how many requests are in flight
+		RequestParams.bIgnoreMaxRequests = true;
+
+		TUniquePtr<FHttpBenchmarkOperation> Operation = MakeUnique<FHttpBenchmarkOperation>(HostCandidate, LocalClient->TryCreateRequest(RequestParams), [&SortedHostList, &Mutex, &HostBenchmarkDone, &FailedBenchmarkAttempts, &bAllHostsFailed, &CountOfBenchmarksToRun](const IHttpResponse& Response, FString& Host)
+		{
+			if (Response.GetStatusCode() != 200)
+			{
+				FailedBenchmarkAttempts += 1;
+				if (FailedBenchmarkAttempts >= CountOfBenchmarksToRun)
+				{
+					bAllHostsFailed = true;
+					HostBenchmarkDone.Notify();
+				}
+
+				return;
+			}
+			// grab stats and store in dictionary, needs to get the response for this
+			const FHttpResponseStats& Stats = Response.GetStats();
+
+			FLatencySortableHost HostLatencyRecord;
+			HostLatencyRecord.Latency = Stats.GetLatency();
+			HostLatencyRecord.Host = Host;
+
+			{
+				TUniqueLock Lock(Mutex);
+				SortedHostList.Add(HostLatencyRecord);
+			}
+
+			HostBenchmarkDone.Notify();
+		});
+
+		Operation->SendAsync();
+		Operations.Emplace(MoveTemp(Operation));
+	}
+	bool bValidHostFound = HostBenchmarkDone.WaitFor(FMonotonicTimeSpan::FromSeconds(WarningTimeoutSeconds));
+
+	// all benchmarks failed
+	if (bAllHostsFailed)
+	{
+		UE_LOG(LogHttpHostBuilder, Warning, TEXT("No valid host found as all benchmark attempts had errors"));
+
+		OutHost.Reset();
+		OutHost.Append(InHostCandidates[0]);
+		OutLatency = 0.0;
+
+		return false;
+	}
+
+	{
+		ON_SCOPE_EXIT {
+			// cancel any outstanding benchmark request
+			for (const TUniquePtr<FHttpBenchmarkOperation>& Operation : Operations)
+			{
+				Operation->Cancel();
+			}
+		};
+
+		if (!bValidHostFound)
+		{
+			// warn that benchmarking is taking a lot of time
+			UE_LOG(LogHttpHostBuilder, Warning, TEXT("HTTP Benchmarking is slow, continuing to wait to determine ideal host..."));
+
+			bValidHostFound = HostBenchmarkDone.WaitFor(FMonotonicTimeSpan::FromSeconds(TimeoutSeconds));
+			if (!bValidHostFound)
+			{
+				UE_LOG(LogHttpHostBuilder, Warning, TEXT("No valid host found while benchmarking after timeout was reached"));
+
+				OutHost.Reset();
+				OutHost.Append(InHostCandidates[0]);
+				OutLatency = 0.0;
+				
+				return false;
+			}
+		}
+		
+		{
+			TUniqueLock Lock(Mutex);
+
+			FLatencySortableHost* MinElement = Algo::MinElementBy(SortedHostList, &FLatencySortableHost::Latency);
+
+			if (MinElement != nullptr)
+			{
+				OutHost.Reset();
+				OutHost.Append(MinElement->Host);
+				OutLatency = MinElement->Latency;
+			}
+			else
+			{
+				// no hosts were found
+				OutHost.Reset();
+				OutHost.Append(InHostCandidates[0]);
+				OutLatency = 0.0;
+
+				UE_LOG(LogHttpHostBuilder, Warning, TEXT("Failed to determine fastest host option because we had no valid options"));
+			}
+		}
+	}
+
+	const double EndTime = FPlatformTime::Seconds();
+	const double BenchmarkingDuration = EndTime - StartTime;
+
+	UE_LOG(LogHttpHostBuilder, Display, TEXT("Resolved to using host '%hs' based on HTTP benchmark with a estimated latency of '%.0fms'. Spent %.0fms doing benchmarking."), *OutHost, OutLatency * 1000, BenchmarkingDuration * 1000);
+	return true;
+}
+
+} // UE
\ No newline at end of file
