Remove unnecessary API declarations on DevHttp headers.

diff --git a/Engine/Source/Developer/DevHttp/Internal/Http/HttpClient.h b/Engine/Source/Developer/DevHttp/Internal/Http/HttpClient.h
--- a/Engine/Source/Developer/DevHttp/Internal/Http/HttpClient.h
+++ b/Engine/Source/Developer/DevHttp/Internal/Http/HttpClient.h
@@ -460,7 +460,7 @@ public:
 	explicit DEVHTTP_API FHttpByteArrayReceiver(TArray64<uint8>& OutArray, IHttpReceiver* Next = nullptr);
 
 	DEVHTTP_API IHttpReceiver* OnBody(IHttpResponse& Response, FMemoryView& Data) final;
-	DEVHTTP_API IHttpReceiver* OnComplete(IHttpResponse& Response) final { return Next; }
+	IHttpReceiver* OnComplete(IHttpResponse& Response) final { return Next; }
 
 private:
 	TArray64<uint8>& Array;
@@ -474,7 +474,7 @@ class FHttpRequestQueue final
 {
 public:
 	/** Create a null request queue. Asserts if CreateRequest is called. */
-	DEVHTTP_API FHttpRequestQueue() = default;
+	FHttpRequestQueue() = default;
 
 	DEVHTTP_API FHttpRequestQueue(IHttpConnectionPool& ConnectionPool, const FHttpClientParams& ClientParams);
 
