Fix to Curl http client in DevHttp module to match the documented behavior where status code is 0 if there is an error.

diff --git a/Engine/Source/Developer/DevHttp/Private/CurlHttpClient.cpp b/Engine/Source/Developer/DevHttp/Private/CurlHttpClient.cpp
--- a/Engine/Source/Developer/DevHttp/Private/CurlHttpClient.cpp
+++ b/Engine/Source/Developer/DevHttp/Private/CurlHttpClient.cpp
@@ -304,7 +304,7 @@ public:
 
 	FAnsiStringView GetUri() const final { return Uri; }
 	EHttpMethod GetMethod() const final { return Method; }
-	int32 GetStatusCode() const final { return StatusCode; }
+	int32 GetStatusCode() const final;
 	EHttpErrorCode GetErrorCode() const final { return ErrorCode; }
 	FAnsiStringView GetError() const final { return Error; }
 	const FHttpResponseStats& GetStats() const final;
@@ -1241,6 +1241,11 @@ void FCurlHttpResponse::Cancel()
 	}
 }
 
+int32 FCurlHttpResponse::GetStatusCode() const
+{
+	return ErrorCode == EHttpErrorCode::None ? StatusCode.load(std::memory_order_relaxed) : 0;
+}
+
 const FHttpResponseStats& FCurlHttpResponse::GetStats() const
 {
 	checkf(!Curl, TEXT("Stats accessed before they are available for %s"), *WriteToString<256>(*this));
