Add a convenience method to get the content length from an http response in DevHttp.

diff --git a/Engine/Source/Developer/DevHttp/Internal/Http/HttpClient.h b/Engine/Source/Developer/DevHttp/Internal/Http/HttpClient.h
--- a/Engine/Source/Developer/DevHttp/Internal/Http/HttpClient.h
+++ b/Engine/Source/Developer/DevHttp/Internal/Http/HttpClient.h
@@ -363,6 +363,9 @@ public:
 	/** Returns the content type of the response, or Any if there was no matching header. */
 	[[nodiscard]] DEVHTTP_API EHttpMediaType GetContentType() const;
 
+	/** Returns the content length of the response, or MAX_uint64 if there was no matching header. */
+	[[nodiscard]] virtual uint64 GetContentLength() const = 0;
+
 protected:
 	friend Http::Private::THttpDestroyer<IHttpResponse>;
 	virtual ~IHttpResponse() = default;
diff --git a/Engine/Source/Developer/DevHttp/Private/CurlHttpClient.cpp b/Engine/Source/Developer/DevHttp/Private/CurlHttpClient.cpp
--- a/Engine/Source/Developer/DevHttp/Private/CurlHttpClient.cpp
+++ b/Engine/Source/Developer/DevHttp/Private/CurlHttpClient.cpp
@@ -309,6 +309,7 @@ public:
 	FAnsiStringView GetError() const final { return Error; }
 	const FHttpResponseStats& GetStats() const final;
 	TConstArrayView<FAnsiStringView> GetAllHeaders() const final { return HeaderViews; }
+	uint64 GetContentLength() const final;
 
 	void AddRef() const final;
 	void Release() const final;
@@ -1252,6 +1253,13 @@ const FHttpResponseStats& FCurlHttpResponse::GetStats() const
 	return Stats;
 }
 
+uint64 FCurlHttpResponse::GetContentLength() const
+{
+	curl_off_t Value;
+	curl_easy_getinfo(Curl, CURLINFO_CONTENT_LENGTH_DOWNLOAD_T, &Value);
+	return uint64(Value);
+}
+
 size_t FCurlHttpResponse::CurlDebug(CURL* Curl, curl_infotype DebugInfoType, char* DebugInfo, size_t DebugInfoSize, void* Param)
 {
 	if (DebugInfoType != CURLINFO_TEXT && DebugInfoType != CURLINFO_HEADER_IN)
