Changing approach for the Content-Length retrieval mechanism in the HttpClient in DevHttp.  The Curl-native method is only usable before the request is complete, so use a generic method of finding the header and converting the value from string to numeral.

diff --git a/Engine/Source/Developer/DevHttp/Internal/Http/HttpClient.h b/Engine/Source/Developer/DevHttp/Internal/Http/HttpClient.h
--- a/Engine/Source/Developer/DevHttp/Internal/Http/HttpClient.h
+++ b/Engine/Source/Developer/DevHttp/Internal/Http/HttpClient.h
@@ -364,7 +364,7 @@ public:
 	[[nodiscard]] DEVHTTP_API EHttpMediaType GetContentType() const;
 
 	/** Returns the content length of the response, or MAX_uint64 if there was no matching header. */
-	[[nodiscard]] virtual uint64 GetContentLength() const = 0;
+	[[nodiscard]] DEVHTTP_API uint64 GetContentLength() const;
 
 protected:
 	friend Http::Private::THttpDestroyer<IHttpResponse>;
diff --git a/Engine/Source/Developer/DevHttp/Private/CurlHttpClient.cpp b/Engine/Source/Developer/DevHttp/Private/CurlHttpClient.cpp
--- a/Engine/Source/Developer/DevHttp/Private/CurlHttpClient.cpp
+++ b/Engine/Source/Developer/DevHttp/Private/CurlHttpClient.cpp
@@ -309,7 +309,6 @@ public:
 	FAnsiStringView GetError() const final { return Error; }
 	const FHttpResponseStats& GetStats() const final;
 	TConstArrayView<FAnsiStringView> GetAllHeaders() const final { return HeaderViews; }
-	uint64 GetContentLength() const final;
 
 	void AddRef() const final;
 	void Release() const final;
@@ -1253,13 +1252,6 @@ const FHttpResponseStats& FCurlHttpResponse::GetStats() const
 	return Stats;
 }
 
-uint64 FCurlHttpResponse::GetContentLength() const
-{
-	curl_off_t Value;
-	curl_easy_getinfo(Curl, CURLINFO_CONTENT_LENGTH_DOWNLOAD_T, &Value);
-	return uint64(Value);
-}
-
 size_t FCurlHttpResponse::CurlDebug(CURL* Curl, curl_infotype DebugInfoType, char* DebugInfo, size_t DebugInfoSize, void* Param)
 {
 	if (DebugInfoType != CURLINFO_TEXT && DebugInfoType != CURLINFO_HEADER_IN)
diff --git a/Engine/Source/Developer/DevHttp/Private/HttpClient.cpp b/Engine/Source/Developer/DevHttp/Private/HttpClient.cpp
--- a/Engine/Source/Developer/DevHttp/Private/HttpClient.cpp
+++ b/Engine/Source/Developer/DevHttp/Private/HttpClient.cpp
@@ -8,6 +8,7 @@
 #include "Misc/AsciiSet.h"
 #include "ProfilingDebugging/CpuProfilerTrace.h"
 #include "String/Find.h"
+#include "String/LexFromString.h"
 
 namespace UE
 {
@@ -289,6 +290,19 @@ EHttpMediaType IHttpResponse::GetContentType() const
 	return MediaType;
 }
 
+uint64 IHttpResponse::GetContentLength() const
+{
+	const FAnsiStringView ContentLength = GetHeader(ANSITEXTVIEW("Content-Length"));
+	const FAnsiStringView ContentLengthNoParams = FAsciiSet::FindPrefixWithout(ContentLength, "; \t");
+	if (ContentLengthNoParams.IsEmpty())
+	{
+		return -1;
+	}
+	uint64 ContentLengthValue = static_cast<uint64>(-1);
+	LexFromString(ContentLengthValue, ContentLengthNoParams);
+	return ContentLengthValue;
+}
+
 FStringBuilderBase& operator<<(FStringBuilderBase& Builder, const IHttpResponse& Response)
 {
 	Builder << LexToString(Response.GetMethod()) << TEXT(' ') << Response.GetUri();
