Allow seeking to the end of a response body without returning CURL_SEEKFUNC_FAIL in the Curl implementation in DevHttp module.

diff --git a/Engine/Source/Developer/DevHttp/Private/CurlHttpClient.cpp b/Engine/Source/Developer/DevHttp/Private/CurlHttpClient.cpp
--- a/Engine/Source/Developer/DevHttp/Private/CurlHttpClient.cpp
+++ b/Engine/Source/Developer/DevHttp/Private/CurlHttpClient.cpp
@@ -900,7 +900,7 @@ size_t FCurlHttpRequest::CurlSeek(void* Param, curl_off_t Offset, int Origin)
 	case SEEK_END: AbsoluteOffset = TotalSize + Offset; break;
 	}
 
-	if (AbsoluteOffset >= TotalSize)
+	if (AbsoluteOffset > TotalSize)
 	{
 		return CURL_SEEKFUNC_FAIL;
 	}
