DevHttp: Warn when a request is idle when it is destroyed

#rnx

diff --git a/Engine/Source/Developer/DevHttp/Private/CurlHttpClient.cpp b/Engine/Source/Developer/DevHttp/Private/CurlHttpClient.cpp
--- a/Engine/Source/Developer/DevHttp/Private/CurlHttpClient.cpp
+++ b/Engine/Source/Developer/DevHttp/Private/CurlHttpClient.cpp
@@ -10,6 +10,7 @@
 #include "Containers/StringView.h"
 #include "HAL/PlatformProcess.h"
 #include "HAL/Thread.h"
+#include "Logging/StructuredLog.h"
 #include "Memory/CompositeBuffer.h"
 #include "Memory/MemoryView.h"
 #include "Misc/App.h"
@@ -716,6 +717,13 @@ FCurlHttpRequest::FCurlHttpRequest(FCurlHttpClient& InClient)
 
 FCurlHttpRequest::~FCurlHttpRequest()
 {
+	//CheckIdle(TEXT("~FCurlHttpRequest"));
+	if (Response.load(std::memory_order_relaxed))
+	{
+		UE_LOGFMT(LogHttp, Warning, "{Method} {Uri}: {FunctionName} requires the request to be idle.",
+			LexToString(Method), Uri, TEXT("~FCurlHttpRequest"));
+		FDebug::DumpStackTraceToLog(ELogVerbosity::Display);
+	}
 	curl_easy_cleanup(Curl);
 }
 
