Reduce verbosity of messages in HttpHostBuilder from warning to display verbosity.  Having warnings emitted because we're too far from any zenserver in our configured timeout is causing issues for systems that treat warnings as errors.

#Synced-CL 48523383

diff --git a/Engine/Source/Developer/DevHttp/Private/HttpHostBuilder.cpp b/Engine/Source/Developer/DevHttp/Private/HttpHostBuilder.cpp
--- a/Engine/Source/Developer/DevHttp/Private/HttpHostBuilder.cpp
+++ b/Engine/Source/Developer/DevHttp/Private/HttpHostBuilder.cpp
@@ -137,7 +137,7 @@ void FHttpHostBuilder::AddFromEndpoint(const FAnsiStringView HostUrl, const FAns
 
 	if (Response->GetStatusCode() != 200)
 	{
-		UE_LOGFMT(LogHttpHostBuilder, Warning, "Unsuccessful attempt to fetch hosts by endpoint from host: '{Host}'. Status code was: {Code}", Uri, Response->GetStatusCode());
+		UE_LOGFMT(LogHttpHostBuilder, Display, "Unsuccessful attempt to fetch hosts by endpoint from host: '{Host}'. Status code was: {Code}", Uri, Response->GetStatusCode());
 		return;
 	}
 
@@ -274,7 +274,7 @@ bool FHttpHostBuilder::BenchmarkHostList(TConstArrayView<FString> InHostCandidat
 	// all benchmarks failed
 	if (bAllHostsFailed)
 	{
-		UE_LOG(LogHttpHostBuilder, Warning, TEXT("No valid host found as all benchmark attempts had errors"));
+		UE_LOG(LogHttpHostBuilder, Display, TEXT("No valid host found as all benchmark attempts had errors"));
 
 		OutHost.Reset();
 		OutHost.Append(InHostCandidates[0]);
@@ -295,12 +295,12 @@ bool FHttpHostBuilder::BenchmarkHostList(TConstArrayView<FString> InHostCandidat
 		if (!bValidHostFound)
 		{
 			// warn that benchmarking is taking a lot of time
-			UE_LOG(LogHttpHostBuilder, Warning, TEXT("HTTP Benchmarking is slow, continuing to wait to determine ideal host..."));
+			UE_LOG(LogHttpHostBuilder, Display, TEXT("HTTP Benchmarking is slow, continuing to wait to determine ideal host..."));
 
 			bValidHostFound = HostBenchmarkDone.WaitFor(FMonotonicTimeSpan::FromSeconds(TimeoutSeconds));
 			if (!bValidHostFound)
 			{
-				UE_LOG(LogHttpHostBuilder, Warning, TEXT("No valid host found while benchmarking after timeout was reached"));
+				UE_LOG(LogHttpHostBuilder, Display, TEXT("No valid host found while benchmarking after timeout was reached"));
 
 				OutHost.Reset();
 				OutHost.Append(InHostCandidates[0]);
@@ -328,7 +328,7 @@ bool FHttpHostBuilder::BenchmarkHostList(TConstArrayView<FString> InHostCandidat
 				OutHost.Append(InHostCandidates[0]);
 				OutLatency = 0.0;
 
-				UE_LOG(LogHttpHostBuilder, Warning, TEXT("Failed to determine fastest host option because we had no valid options"));
+				UE_LOG(LogHttpHostBuilder, Display, TEXT("Failed to determine fastest host option because we had no valid options"));
 			}
 		}
 	}
