Provide more verbose data if we fail to communicate with a zen DDC server.

#Synced-CL 49059893

diff --git a/Engine/Source/Developer/DerivedDataCache/Private/ZenCacheStore.cpp b/Engine/Source/Developer/DerivedDataCache/Private/ZenCacheStore.cpp
--- a/Engine/Source/Developer/DerivedDataCache/Private/ZenCacheStore.cpp
+++ b/Engine/Source/Developer/DerivedDataCache/Private/ZenCacheStore.cpp
@@ -1831,15 +1831,21 @@ void FZenCacheStore::Initialize(const FZenCacheStoreParams& Params)
 			if (ZenService.GetInstance().IsServiceRunningLocally())
 			{
 				UE_LOG(LogDerivedDataCache, Warning,
-					TEXT("%s: Unable to reach ZenServer HTTP service at %s with namespace %s. Status: %d . Response: %s"),
-					*GetName(), ZenServiceEndpoint.GetName().GetData(), *Namespace, ReadinessResponse->GetStatusCode(),
+					TEXT("%s: Unable to reach ZenServer HTTP service at %s with namespace %s. ErrorCode: %s, Error: %s, Status: %d, Response: %s"),
+					*GetName(), ZenServiceEndpoint.GetName().GetData(), *Namespace,
+					*WriteToString<16>(LexToString(ReadinessResponse->GetErrorCode())),
+					ReadinessResponse->GetStatusCode() == 0 ? TEXT("") : *WriteToString<64>(ReadinessResponse->GetError()),
+					ReadinessResponse->GetStatusCode(),
 					*ResponseString);
 			}
 			else
 			{
 				UE_LOG(LogDerivedDataCache, Display,
-					TEXT("%s: Unable to reach ZenServer HTTP service at %s with namespace %s. Status: %d . Response: %s"),
-					*GetName(), ZenServiceEndpoint.GetName().GetData(), *Namespace, ReadinessResponse->GetStatusCode(),
+					TEXT("%s: Unable to reach ZenServer HTTP service at %s with namespace %s. ErrorCode: %s, Error: %s, Status: %d, Response: %s"),
+					*GetName(), ZenServiceEndpoint.GetName().GetData(), *Namespace,
+					*WriteToString<16>(LexToString(ReadinessResponse->GetErrorCode())),
+					ReadinessResponse->GetStatusCode() == 0 ? TEXT("") : *WriteToString<64>(ReadinessResponse->GetError()),
+					ReadinessResponse->GetStatusCode(),
 					*ResponseString);
 			}
 		}
diff --git a/Engine/Source/Developer/DevHttp/Internal/Http/HttpClient.h b/Engine/Source/Developer/DevHttp/Internal/Http/HttpClient.h
--- a/Engine/Source/Developer/DevHttp/Internal/Http/HttpClient.h
+++ b/Engine/Source/Developer/DevHttp/Internal/Http/HttpClient.h
@@ -119,6 +119,9 @@ enum class EHttpErrorCode : uint8
 	TimedOut,
 };
 
+DEVHTTP_API FAnsiStringView LexToString(EHttpErrorCode ErrorCode);
+DEVHTTP_API bool TryLexFromString(EHttpErrorCode& OutErrorCode, FAnsiStringView View);
+
 class IHttpManager
 {
 public:
diff --git a/Engine/Source/Developer/DevHttp/Private/HttpClient.cpp b/Engine/Source/Developer/DevHttp/Private/HttpClient.cpp
--- a/Engine/Source/Developer/DevHttp/Private/HttpClient.cpp
+++ b/Engine/Source/Developer/DevHttp/Private/HttpClient.cpp
@@ -145,6 +145,64 @@ bool TryLexFromString(EHttpMediaType& OutMediaType, const FAnsiStringView View)
 	return true;
 }
 
+FAnsiStringView LexToString(EHttpErrorCode ErrorCode)
+{
+	switch (ErrorCode)
+	{
+	case EHttpErrorCode::None:                return ANSITEXTVIEW("None");
+	case EHttpErrorCode::Unknown:             return ANSITEXTVIEW("Unknown");
+	case EHttpErrorCode::Canceled:            return ANSITEXTVIEW("Canceled");
+	case EHttpErrorCode::ResolveHost:         return ANSITEXTVIEW("ResolveHost");
+	case EHttpErrorCode::Connect:             return ANSITEXTVIEW("Connect");
+	case EHttpErrorCode::TlsConnect:          return ANSITEXTVIEW("TlsConnect");
+	case EHttpErrorCode::TlsPeerVerification: return ANSITEXTVIEW("TlsPeerVerification");
+	case EHttpErrorCode::TimedOut:            return ANSITEXTVIEW("TimedOut");
+	default:                                  return ANSITEXTVIEW("Invalid");
+	}
+
+}
+
+bool TryLexFromString(EHttpErrorCode& OutErrorCode, FAnsiStringView View)
+{
+	if (View == ANSITEXTVIEW("None"))
+	{
+		OutErrorCode = EHttpErrorCode::None;
+	}
+	else if (View == ANSITEXTVIEW("Unknown"))
+	{
+		OutErrorCode = EHttpErrorCode::Unknown;
+	}
+	else if (View == ANSITEXTVIEW("Canceled"))
+	{
+		OutErrorCode = EHttpErrorCode::Canceled;
+	}
+	else if (View == ANSITEXTVIEW("ResolveHost"))
+	{
+		OutErrorCode = EHttpErrorCode::ResolveHost;
+	}
+	else if (View == ANSITEXTVIEW("Connect"))
+	{
+		OutErrorCode = EHttpErrorCode::Connect;
+	}
+	else if (View == ANSITEXTVIEW("TlsConnect"))
+	{
+		OutErrorCode = EHttpErrorCode::TlsConnect;
+	}
+	else if (View == ANSITEXTVIEW("TlsPeerVerification"))
+	{
+		OutErrorCode = EHttpErrorCode::TlsPeerVerification;
+	}
+	else if (View == ANSITEXTVIEW("TimedOut"))
+	{
+		OutErrorCode = EHttpErrorCode::TimedOut;
+	}
+	else
+	{
+		return false;
+	}
+	return true;
+}
+
 static FAnsiStringView LexStatusCodeToString(int HttpCode)
 {
 	switch (HttpCode)
