When doing host evaluation, we should allow as many connections as we have hosts to evaluate, otherwise we won't be fairly evaluating the full set of hosts.

#Synced-CL 51289657

diff --git a/Engine/Source/Developer/DevHttp/Private/HttpHostBuilder.cpp b/Engine/Source/Developer/DevHttp/Private/HttpHostBuilder.cpp
--- a/Engine/Source/Developer/DevHttp/Private/HttpHostBuilder.cpp
+++ b/Engine/Source/Developer/DevHttp/Private/HttpHostBuilder.cpp
@@ -209,7 +209,7 @@ bool FHttpHostBuilder::BenchmarkHostList(TConstArrayView<FString> InHostCandidat
 {
 	const double StartTime = FPlatformTime::Seconds();
 
-	constexpr uint32 MaxTotalConnections = 8;
+	const uint32 MaxTotalConnections = InHostCandidates.Num();
 	FHttpConnectionPoolParams ConnectionPoolParams;
 	ConnectionPoolParams.MaxConnections = MaxTotalConnections;
 	ConnectionPoolParams.MinConnections = MaxTotalConnections;
