Fix race condition in HostBuilder involving the counter for failed connection attempts.

#Synced-CL 51529686

diff --git a/Engine/Source/Developer/DevHttp/Private/HttpHostBuilder.cpp b/Engine/Source/Developer/DevHttp/Private/HttpHostBuilder.cpp
--- a/Engine/Source/Developer/DevHttp/Private/HttpHostBuilder.cpp
+++ b/Engine/Source/Developer/DevHttp/Private/HttpHostBuilder.cpp
@@ -228,9 +228,9 @@ bool FHttpHostBuilder::BenchmarkHostList(TConstArrayView<FString> InHostCandidat
 	FMutex Mutex;
 	FManualResetEvent HostBenchmarkDone;
 
-	int32 FailedBenchmarkAttempts = 0;
+	std::atomic<int32> FailedBenchmarkAttempts = 0;
 	bool bAllHostsFailed = false;
-	int32 CountOfBenchmarksToRun = InHostCandidates.Num();
+	const int32 CountOfBenchmarksToRun = InHostCandidates.Num();
 	for (const FString& HostCandidate : InHostCandidates)
 	{
 		// Benchmark each of the hosts
@@ -242,8 +242,7 @@ bool FHttpHostBuilder::BenchmarkHostList(TConstArrayView<FString> InHostCandidat
 		{
 			if (Response.GetStatusCode() != 200)
 			{
-				FailedBenchmarkAttempts += 1;
-				if (FailedBenchmarkAttempts >= CountOfBenchmarksToRun)
+				if (FailedBenchmarkAttempts.fetch_add(1, std::memory_order_relaxed) + 1 == CountOfBenchmarksToRun)
 				{
 					bAllHostsFailed = true;
 					HostBenchmarkDone.Notify();
